Publish Modulation
- Fix thread safety issues
- Fix auto-activation/deactivation in early form (WIP: Need to add source counting.  Realized needs a refactor and cannot rely on source ids.  Virtualized sources should be able to act as a "source" in count)

diff --git a/Engine/Plugins/Runtime/AudioModulation/AudioModulation.uplugin b/Engine/Plugins/Runtime/AudioModulation/AudioModulation.uplugin
new file mode 100644
--- /dev/null
+++ b/Engine/Plugins/Runtime/AudioModulation/AudioModulation.uplugin
@@ -0,0 +1,37 @@
+{
+	"FileVersion": 3,
+	"Version": 1,
+	"VersionName": "1.0",
+	"FriendlyName": "Audio Modulation",
+	"Description": "Plugin provides an interface for modulating common parameters in audio.",
+	"Category": "Audio",
+	"CreatedBy": "Epic Games, Inc.",
+	"CreatedByURL": "http://epicgames.com",
+	"DocsURL": "",
+	"MarketplaceURL": "",
+	"SupportURL": "",
+	"EnabledByDefault": false,
+	"CanContainContent": false,
+	"IsBetaVersion": false,
+	"Installed": false,
+	"Modules": [
+		{
+			"Name": "AudioModulation",
+			"Type": "Runtime",
+			"LoadingPhase": "Default",
+			"WhitelistPlatforms": [
+				"Win64",
+				"Lumin",
+				"Mac"
+			]
+		},
+		{
+			"Name": "AudioModulationEditor",
+			"Type": "Editor",
+			"LoadingPhase": "PostEngineInit",
+			"WhitelistPlatforms": [
+				"Win64"
+			]
+		}
+	]
+}
\ No newline at end of file
diff --git a/Engine/Plugins/Runtime/AudioModulation/Source/AudioModulation/AudioModulation.Build.cs b/Engine/Plugins/Runtime/AudioModulation/Source/AudioModulation/AudioModulation.Build.cs
new file mode 100644
--- /dev/null
+++ b/Engine/Plugins/Runtime/AudioModulation/Source/AudioModulation/AudioModulation.Build.cs
@@ -0,0 +1,30 @@
+// Copyright 1998-2019 Epic Games, Inc. All Rights Reserved.
+
+namespace UnrealBuildTool.Rules
+{
+	public class AudioModulation : ModuleRules
+	{
+		public AudioModulation(ReadOnlyTargetRules Target) : base(Target)
+		{
+			OptimizeCode = CodeOptimization.Never;
+
+			PublicDependencyModuleNames.AddRange(
+				new string[]
+				{
+					"Core"
+				}
+			);
+
+			PrivateDependencyModuleNames.AddRange(
+				new string[]
+				{
+					"CoreUObject",
+					"Engine",
+					"AudioMixer"
+				}
+			);
+
+			PublicDefinitions.Add("WITH_AUDIOMODULATION=1");
+		}
+	}
+}
\ No newline at end of file
diff --git a/Engine/Plugins/Runtime/AudioModulation/Source/AudioModulation/Private/AudioModulation.cpp b/Engine/Plugins/Runtime/AudioModulation/Source/AudioModulation/Private/AudioModulation.cpp
new file mode 100644
--- /dev/null
+++ b/Engine/Plugins/Runtime/AudioModulation/Source/AudioModulation/Private/AudioModulation.cpp
@@ -0,0 +1,100 @@
+// Copyright 1998-2019 Epic Games, Inc. All Rights Reserved.
+#include "AudioModulation.h"
+
+#include "AudioModulationInternal.h"
+#include "AudioModulationLogging.h"
+#include "CanvasTypes.h"
+#include "Features/IModularFeatures.h"
+#include "IAudioExtensionPlugin.h"
+#include "Modules/ModuleManager.h"
+#include "SoundModulationPatch.h"
+#include "UnrealClient.h"
+
+
+DEFINE_STAT(STAT_AudioModulationProcessControls);
+DEFINE_STAT(STAT_AudioModulationProcessModulators);
+
+namespace AudioModulation
+{
+	FAudioModulation::FAudioModulation()
+	{
+		Impl = TUniquePtr<FAudioModulationImpl>(new FAudioModulationImpl());
+	}
+
+	FAudioModulationImpl* FAudioModulation::GetImpl()
+	{
+		return Impl.Get();
+	}
+
+	void FAudioModulation::Initialize(const FAudioPluginInitializationParams& InitializationParams)
+	{
+		Impl->Initialize(InitializationParams);
+	}
+
+#if WITH_EDITOR
+	void FAudioModulation::OnEditSource(const USoundModulationPluginSourceSettingsBase& Settings)
+	{
+		Impl->OnEditSource(Settings);
+	}
+#endif // WITH_EDITOR
+
+	void FAudioModulation::OnInitSource(const uint32 SourceId, const FName& AudioComponentUserId, const uint32 NumChannels, USoundModulationPluginSourceSettingsBase* Settings)
+	{
+		Impl->OnInitSource(SourceId, AudioComponentUserId, NumChannels, Settings);
+	}
+
+	void FAudioModulation::OnReleaseSource(const uint32 SourceId)
+	{
+		Impl->OnReleaseSource(SourceId);
+	}
+
+#if !UE_BUILD_SHIPPING
+	bool FAudioModulation::OnPostHelp(FCommonViewportClient* ViewportClient, const TCHAR* Stream)
+	{
+		return Impl->OnPostHelp(ViewportClient, Stream);
+	}
+
+	int32 FAudioModulation::OnRenderStat(FViewport* Viewport, FCanvas* Canvas, int32 X, int32 Y, const UFont& Font, const FVector* ViewLocation, const FRotator* ViewRotation)
+	{
+		return Impl->OnRenderStat(Viewport, Canvas, X, Y, Font, ViewLocation, ViewRotation);
+	}
+
+	bool FAudioModulation::OnToggleStat(FCommonViewportClient* ViewportClient, const TCHAR* Stream)
+	{
+		return Impl->OnToggleStat(ViewportClient, Stream);
+	}
+#endif // !UE_BUILD_SHIPPING
+
+	void FAudioModulation::ProcessControls(const uint32 SourceId, FSoundModulationControls& Controls)
+	{
+		SCOPE_CYCLE_COUNTER(STAT_AudioModulationProcessControls);
+		Impl->ProcessControls(SourceId, Controls);
+	}
+
+	void FAudioModulation::ProcessModulators(const float Elapsed)
+	{
+		SCOPE_CYCLE_COUNTER(STAT_AudioModulationProcessModulators);
+		Impl->ProcessModulators(Elapsed);
+	}
+} // namespace AudioModulation
+
+TAudioModulationPtr FAudioModulationPluginFactory::CreateNewModulationPlugin(FAudioDevice* OwningDevice)
+{
+	return TAudioModulationPtr(new AudioModulation::FAudioModulation());
+}
+
+void FAudioModulationModule::StartupModule()
+{
+	UE_LOG(LogAudioModulation, Log, TEXT("Starting Audio Modulation Module"));
+
+	IModularFeatures::Get().RegisterModularFeature(FAudioModulationPluginFactory::GetModularFeatureName(), &ModulationPluginFactory);
+}
+
+void FAudioModulationModule::ShutdownModule()
+{
+	UE_LOG(LogAudioModulation, Log, TEXT("Shutting Down Audio Modulation Module"));
+
+	IModularFeatures::Get().UnregisterModularFeature(FAudioModulationPluginFactory::GetModularFeatureName(), &ModulationPluginFactory);
+}
+
+IMPLEMENT_MODULE(FAudioModulationModule, AudioModulation);
\ No newline at end of file
diff --git a/Engine/Plugins/Runtime/AudioModulation/Source/AudioModulation/Private/AudioModulationDebugger.cpp b/Engine/Plugins/Runtime/AudioModulation/Source/AudioModulation/Private/AudioModulationDebugger.cpp
new file mode 100644
--- /dev/null
+++ b/Engine/Plugins/Runtime/AudioModulation/Source/AudioModulation/Private/AudioModulationDebugger.cpp
@@ -0,0 +1,406 @@
+// Copyright 1998-2019 Epic Games, Inc. All Rights Reserved.
+#include "AudioModulationDebugger.h"
+
+#if WITH_AUDIOMODULATION
+#if !UE_BUILD_SHIPPING
+#include "CoreMinimal.h"
+#include "AudioModulationLogging.h"
+#include "AudioThread.h"
+#include "Engine/Canvas.h"
+#include "Engine/Engine.h"
+#include "IAudioExtensionPlugin.h"
+#include "Misc/CoreDelegates.h"
+#include "SoundModulationPatch.h"
+#include "SoundModulationValue.h"
+
+
+namespace AudioModulation
+{
+	namespace
+	{
+		const int32 MaxNameLength = 40;
+		const int32 XIndent       = 36;
+
+		FColor GetUnitRenderColor(const float Value, const FVector2D& Range)
+		{
+			return Value > Range.Y || Range.X < Value
+				? FColor::Red
+				: FColor::Green;
+		}
+
+		int32 RenderStatLFO(const TArray<FModulatorLFOProxy>& FilteredLFOs, FCanvas& Canvas, int32 X, int32 Y, const UFont& Font)
+		{
+			int32 Height = 12;
+			int32 Width = 12;
+			Font.GetStringHeightAndWidth(TEXT("@"), Height, Width);
+
+			// Determine minimum width of cells
+			static const FString NameHeader      = TEXT("Name");
+			static const FString ValueHeader     = TEXT("Value");
+			static const FString AmplitudeHeader = TEXT("Amplitude");
+			static const FString FrequencyHeader = TEXT("Frequency");
+			static const FString OffsetHeader    = TEXT("Offset");
+			static const TArray<int32> StaticCellWidths =
+			{
+				NameHeader.Len(),
+				ValueHeader.Len(),
+				AmplitudeHeader.Len(),
+				FrequencyHeader.Len(),
+				OffsetHeader.Len(),
+			};
+
+			int32 CellWidth = FMath::Max(StaticCellWidths);
+			for (const FModulatorLFOProxy& LFOProxy : FilteredLFOs)
+			{
+				CellWidth = FMath::Max(CellWidth, LFOProxy.GetName().Len());
+			}
+
+			Canvas.DrawShadowedString(X + XIndent, Y, TEXT("Active LFOs:"), &Font, FColor::Red);
+			Y += Height;
+			X += XIndent;
+
+			int32 RowX = X;
+			Canvas.DrawShadowedString(RowX, Y, *NameHeader, &Font, FColor::White);
+			RowX += Width * CellWidth;
+
+			Canvas.DrawShadowedString(RowX, Y, *ValueHeader, &Font, FColor::White);
+			RowX += Width * CellWidth;
+
+			Canvas.DrawShadowedString(RowX, Y, *AmplitudeHeader, &Font, FColor::White);
+			RowX += Width * CellWidth;
+
+			Canvas.DrawShadowedString(RowX, Y, *FrequencyHeader, &Font, FColor::White);
+			RowX += Width * CellWidth;
+
+			Canvas.DrawShadowedString(RowX, Y, *OffsetHeader, &Font, FColor::White);
+			RowX += Width * CellWidth;
+
+			Y += Height;
+			for (const FModulatorLFOProxy& LFOPair : FilteredLFOs)
+			{
+				RowX = X;
+
+				FString Name = LFOPair.GetName().Left(MaxNameLength);
+				if (Name.Len() < MaxNameLength)
+				{
+					Name = Name.RightPad(MaxNameLength - Name.Len());
+				}
+
+				const FString ValueString  = FString::Printf(TEXT("%.6f"), LFOPair.GetValue());
+				const FString AmpString    = FString::Printf(TEXT("%.6f"), LFOPair.GetAmplitude());
+				const FString FreqString   = FString::Printf(TEXT("%.6f"), LFOPair.GetFreq());
+				const FString OffsetString = FString::Printf(TEXT("%.6f"), LFOPair.GetOffset());
+
+				Canvas.DrawShadowedString(RowX, Y, *Name, &Font, FColor::Green);
+				RowX += Width * CellWidth;
+
+				Canvas.DrawShadowedString(RowX, Y, *ValueString, &Font, FColor::Green);
+				RowX += Width * CellWidth;
+
+				Canvas.DrawShadowedString(RowX, Y, *AmpString, &Font, FColor::Green);
+				RowX += Width * CellWidth;
+
+				Canvas.DrawShadowedString(RowX, Y, *FreqString, &Font, FColor::Green);
+				RowX += Width * CellWidth;
+
+				Canvas.DrawShadowedString(RowX, Y, *OffsetString, &Font, FColor::Green);
+				RowX += Width * CellWidth;
+
+				Y += Height;
+			}
+
+			return Y;
+		}
+
+		int32 RenderStatMixMatrix(const TArray<FModulatorBusMixProxy>& FilteredMixes, const TArray<FModulatorBusProxy>& FilteredBuses, FCanvas& Canvas, int32 X, int32 Y, const UFont& Font)
+		{
+			int32 Height = 12;
+			int32 Width = 12;
+			Font.GetStringHeightAndWidth(TEXT("@"), Height, Width);
+			X += XIndent;
+
+			Canvas.DrawShadowedString(X, Y, TEXT("Bus Mix Matrix:"), &Font, FColor::Red);
+			Y += Height;
+
+			// Determine minimum width of cells
+			static const FString MixSubTotalHeader = TEXT("Mix");
+			static const FString LFOSubTotalHeader = TEXT("LFO");
+			static const FString TotalHeader       = TEXT("Final");
+			static const TArray<int32> StaticCellWidths =
+			{
+				LFOSubTotalHeader.Len(),
+				MixSubTotalHeader.Len(),
+				TotalHeader.Len(),
+				14 /* Minimum width due to value strings X.XXXX(X.XXXX) */
+			};
+
+			int32 CellWidth = FMath::Max(StaticCellWidths);
+			for (const FModulatorBusMixProxy& BusMix : FilteredMixes)
+			{
+				CellWidth = FMath::Max(CellWidth, BusMix.GetName().Len());
+			}
+			for (const FModulatorBusProxy& BusPair : FilteredBuses)
+			{
+				CellWidth = FMath::Max(CellWidth, BusPair.GetName().Len());
+			}
+
+			// Draw Column Headers
+			int32 RowX = X;
+			for (const FModulatorBusProxy& BusPair : FilteredBuses)
+			{
+				FString Name = BusPair.GetName().Left(MaxNameLength);
+				if (Name.Len() < MaxNameLength)
+				{
+					Name = Name.RightPad(MaxNameLength - Name.Len());
+				}
+
+				RowX += Width * CellWidth; // Add before to leave space for row headers
+				Canvas.DrawShadowedString(RowX, Y, *Name, &Font, FColor::White);
+			}
+
+			// Draw Row Headers
+			int32 ColumnY = Y;
+			for (const FModulatorBusMixProxy& BusMix : FilteredMixes)
+			{
+				FString Name = BusMix.GetName().Left(MaxNameLength);
+				if (Name.Len() < MaxNameLength)
+				{
+					Name = Name.RightPad(MaxNameLength - Name.Len());
+				}
+
+				ColumnY += Height; // Add before to leave space for column headers
+				Canvas.DrawShadowedString(X, ColumnY, *Name, &Font, FColor::White);
+			}
+
+			// Reset Corner of Matrix & Draw Per Bus Data
+			RowX = X;
+			ColumnY = Y;
+			for (const FModulatorBusMixProxy& BusMix : FilteredMixes)
+			{
+				ColumnY += Height; // Add before to leave space for column headers
+				RowX = X;
+
+				for (const FModulatorBusProxy& Bus : FilteredBuses)
+				{
+					RowX += Width * CellWidth; // Add before to leave space for row headers
+
+					float Target = Bus.GetDefaultValue();
+					float Value = Bus.GetDefaultValue();
+					if (const FModulatorBusMixChannelProxy* ChannelProxy = BusMix.Channels.Find(Bus.GetBusId()))
+					{
+						Target = ChannelProxy->Value.TargetValue;
+						Value  = ChannelProxy->Value.GetCurrentValue();
+					}
+
+					if (Target != Value)
+					{
+						Canvas.DrawShadowedString(RowX, ColumnY, *FString::Printf(TEXT("%.4f(%.4f)"), Value, Target), &Font, FColor::Green);
+					}
+					else
+					{
+						Canvas.DrawShadowedString(RowX, ColumnY, *FString::Printf(TEXT("%.4f"), Value), &Font, FColor::Green);
+					}
+				}
+			}
+
+			Y += (FilteredMixes.Num() + 2) * Height; // Add 2, one for header and one for spacing row
+
+			// Draw Sub-Totals & Totals
+			Canvas.DrawShadowedString(X, Y, *MixSubTotalHeader, &Font, FColor::Yellow);
+			RowX = X;
+			for (const FModulatorBusProxy& BusPair : FilteredBuses)
+			{
+				RowX += Width * CellWidth; // Add before to leave space for row headers
+
+				const float Value = BusPair.GetMixValue();
+				if (FMath::IsNaN(Value))
+				{
+					Canvas.DrawShadowedString(RowX, Y, TEXT("N/A"), &Font, FColor::Green);
+				}
+				else
+				{
+					const FColor Color = GetUnitRenderColor(Value, BusPair.GetRange());
+					Canvas.DrawShadowedString(RowX, Y, *FString::Printf(TEXT("%.4f"), Value), &Font, Color);
+				}
+			}
+			Y += Height;
+
+			Canvas.DrawShadowedString(X, Y, *LFOSubTotalHeader, &Font, FColor::Yellow);
+			RowX = X;
+			for (const FModulatorBusProxy& BusPair : FilteredBuses)
+			{
+				RowX += Width * CellWidth; // Add before to leave space for row headers
+				const float Value = BusPair.GetLFOValue();
+				const FColor Color = GetUnitRenderColor(Value, BusPair.GetRange());
+				Canvas.DrawShadowedString(RowX, Y, *FString::Printf(TEXT("%.4f"), Value), &Font, Color);
+			}
+
+			Y += Height;
+			Canvas.DrawShadowedString(X, Y, *TotalHeader, &Font, FColor::Yellow);
+			RowX = X;
+			for (const FModulatorBusProxy& BusPair : FilteredBuses)
+			{
+				RowX += Width * CellWidth; // Add before to leave space for row headers
+				const float Value = BusPair.GetValue();
+				const FColor Color = GetUnitRenderColor(Value, BusPair.GetRange());
+				Canvas.DrawShadowedString(RowX, Y, *FString::Printf(TEXT("%.4f"), Value), &Font, Color);
+			}
+			Y += Height;
+
+			return Y + Height;
+		}
+
+		template <typename T>
+		bool CompareNames(const T& A, const T& B)
+		{
+			return A.GetName() < B.GetName();
+		}
+
+		template <typename T, typename U>
+		void FilterDebugArray(const T& Map, const FString& FilterString, int32 MaxCount, TArray<U>& FilteredArray)
+		{
+			FilteredArray.Reset(MaxCount);
+
+			int32 FilteredItemCount = 0;
+			for (const auto& IdItemPair : Map)
+			{
+				const U& Item = IdItemPair.Value;
+				const bool Filtered = !FilterString.IsEmpty()
+					&& !Item.GetName().Contains(FilterString);
+				if (!Filtered)
+				{
+					FilteredArray.Add(Item);
+					if (++FilteredItemCount >= MaxCount)
+					{
+						return;
+					}
+				}
+			}
+		}
+	} // namespace <>
+
+	FAudioModulationDebugger::FAudioModulationDebugger()
+		: bActive(0)
+		, bShowRenderStatLFO(1)
+		, bShowRenderStatMix(1)
+	{
+	}
+
+	void FAudioModulationDebugger::UpdateDebugData(
+		const BusProxyMap&    ActiveBuses,
+		const BusMixProxyMap& ActiveMixes,
+		const LFOProxyMap&    ActiveLFOs)
+	{
+		check(IsInAudioThread());
+
+		if (!bActive)
+		{
+			return;
+		}
+
+		static const int32 MaxFilteredBuses = 8;
+		TArray<FModulatorBusProxy> InFilteredBuses;
+		FilterDebugArray<BusProxyMap, FModulatorBusProxy>(ActiveBuses, BusStringFilter, MaxFilteredBuses, InFilteredBuses);
+
+		static const int32 MaxFilteredMixes = 16;
+		TArray<FModulatorBusMixProxy> InFilteredMixes;
+		FilterDebugArray<BusMixProxyMap, FModulatorBusMixProxy>(ActiveMixes, MixStringFilter, MaxFilteredMixes, InFilteredMixes);
+
+		static const int32 MaxFilteredLFOs = 8;
+		TArray<FModulatorLFOProxy> InFilteredLFOs;
+		FilterDebugArray<LFOProxyMap, FModulatorLFOProxy>(ActiveLFOs, LFOStringFilter, MaxFilteredLFOs, InFilteredLFOs);
+
+		FAudioThread::RunCommandOnGameThread([this, InFilteredBuses, InFilteredMixes, InFilteredLFOs]()
+		{
+			FilteredBuses = InFilteredBuses;
+			FilteredBuses.Sort(&CompareNames<FModulatorBusProxy>);
+
+			FilteredMixes = InFilteredMixes;
+			FilteredMixes.Sort(&CompareNames<FModulatorBusMixProxy>);
+
+			FilteredLFOs = InFilteredLFOs;
+			FilteredLFOs.Sort(&CompareNames<FModulatorLFOProxy>);
+		});
+	}
+
+	bool FAudioModulationDebugger::OnPostHelp(FCommonViewportClient& ViewportClient, const TCHAR* Stream)
+	{
+		if (GEngine)
+		{
+			static const float TimeShowHelp = 10.0f;
+			GEngine->AddOnScreenDebugMessage(-1, TimeShowHelp, FColor::Yellow, TEXT("-MixFilter: Substring that filters mixes shown in matrix view"));
+			GEngine->AddOnScreenDebugMessage(-1, TimeShowHelp, FColor::Yellow, TEXT("-BusFilter: Substring that filters buses shown in matrix view"));
+			GEngine->AddOnScreenDebugMessage(-1, TimeShowHelp, FColor::Yellow, TEXT("-Matrix: Show bus matrix"));
+			GEngine->AddOnScreenDebugMessage(-1, TimeShowHelp, FColor::Yellow, TEXT("-LFOFilter: Substring that filters LFOs listed"));
+			GEngine->AddOnScreenDebugMessage(-1, TimeShowHelp, FColor::Yellow, TEXT("-LFO: Show LFO debug data"));
+			GEngine->AddOnScreenDebugMessage(-1, TimeShowHelp, FColor::Yellow, TEXT("stat SoundModulators:"));
+		}
+
+		return true;
+	}
+
+	int32 FAudioModulationDebugger::OnRenderStat(FCanvas& Canvas, int32 X, int32 Y, const UFont& Font)
+	{
+		check(IsInGameThread());
+
+		if (bShowRenderStatMix)
+		{
+			Y = RenderStatMixMatrix(FilteredMixes, FilteredBuses, Canvas, X, Y, Font);
+		}
+
+		if (bShowRenderStatLFO)
+		{
+			Y = RenderStatLFO(FilteredLFOs, Canvas, X, Y, Font);
+		}
+
+		return Y;
+	}
+
+	bool FAudioModulationDebugger::OnToggleStat(FCommonViewportClient& ViewportClient, const TCHAR* Stream)
+	{
+		check(IsInGameThread());
+
+		if (bActive == ViewportClient.IsStatEnabled(TEXT("SoundModulators")))
+		{
+			return true;
+		}
+
+		if (!bActive)
+		{
+			if (Stream && Stream[0] != '\0')
+			{
+				const bool bUpdateLFO = FParse::Param(Stream, TEXT("LFO"));
+				const bool bUpdateMix = FParse::Param(Stream, TEXT("Matrix"));
+
+				if (bUpdateLFO || bUpdateMix)
+				{
+					bShowRenderStatLFO = bUpdateLFO;
+					bShowRenderStatMix = bUpdateMix;
+				}
+				else
+				{
+					bShowRenderStatLFO = 1;
+					bShowRenderStatMix = 1;
+				}
+
+				FParse::Value(Stream, TEXT("BusFilter"), BusStringFilter);
+				FParse::Value(Stream, TEXT("LFOFilter"), LFOStringFilter);
+				FParse::Value(Stream, TEXT("MixFilter"), MixStringFilter);
+			}
+			else
+			{
+				bShowRenderStatLFO = 1;
+				bShowRenderStatMix = 1;
+
+				BusStringFilter.Reset();
+				LFOStringFilter.Reset();
+				MixStringFilter.Reset();
+			}
+		}
+
+		bActive = !bActive;
+		return true;
+	}
+} // namespace AudioModulation
+#endif // !UE_BUILD_SHIPPING
+#endif // WITH_AUDIOMODULATION
\ No newline at end of file
diff --git a/Engine/Plugins/Runtime/AudioModulation/Source/AudioModulation/Private/AudioModulationDebugger.h b/Engine/Plugins/Runtime/AudioModulation/Source/AudioModulation/Private/AudioModulationDebugger.h
new file mode 100644
--- /dev/null
+++ b/Engine/Plugins/Runtime/AudioModulation/Source/AudioModulation/Private/AudioModulationDebugger.h
@@ -0,0 +1,43 @@
+// Copyright 1998-2019 Epic Games, Inc. All Rights Reserved.
+#pragma once
+
+#include "CoreMinimal.h"
+
+#if WITH_AUDIOMODULATION
+#if !UE_BUILD_SHIPPING
+#include "SoundModulatorBus.h"
+#include "SoundModulatorBusMix.h"
+#include "SoundModulatorLFO.h"
+
+
+namespace AudioModulation
+{
+	class FAudioModulationDebugger
+	{
+	public:
+		FAudioModulationDebugger();
+
+		void UpdateDebugData(
+			const BusProxyMap&    ActiveBuses,
+			const BusMixProxyMap& ActiveMixes,
+			const LFOProxyMap&    ActiveLFOs);
+		bool OnPostHelp(FCommonViewportClient& ViewportClient, const TCHAR* Stream);
+		int32 OnRenderStat(FCanvas& Canvas, int32 X, int32 Y, const UFont& Font);
+		bool OnToggleStat(FCommonViewportClient& ViewportClient, const TCHAR* Stream);
+
+	private:
+		uint8 bActive : 1;
+		uint8 bShowRenderStatLFO : 1;
+		uint8 bShowRenderStatMix : 1;
+
+		TArray<FModulatorBusProxy>    FilteredBuses;
+		TArray<FModulatorLFOProxy>    FilteredLFOs;
+		TArray<FModulatorBusMixProxy> FilteredMixes;
+
+		FString BusStringFilter;
+		FString LFOStringFilter;
+		FString MixStringFilter;
+	};
+} // namespace AudioModulation
+#endif // !UE_BUILD_SHIPPING
+#endif // WITH_AUDIOMODULATION
diff --git a/Engine/Plugins/Runtime/AudioModulation/Source/AudioModulation/Private/AudioModulationImpl.cpp b/Engine/Plugins/Runtime/AudioModulation/Source/AudioModulation/Private/AudioModulationImpl.cpp
new file mode 100644
--- /dev/null
+++ b/Engine/Plugins/Runtime/AudioModulation/Source/AudioModulation/Private/AudioModulationImpl.cpp
@@ -0,0 +1,505 @@
+// Copyright 1998-2019 Epic Games, Inc. All Rights Reserved.
+#include "AudioModulationInternal.h"
+
+#if WITH_AUDIOMODULATION
+#include "AudioModulationLogging.h"
+#include "AudioThread.h"
+#include "Engine/Engine.h"
+#include "IAudioExtensionPlugin.h"
+#include "Misc/CoreDelegates.h"
+#include "SoundModulatorLFO.h"
+#include "SoundModulationPatch.h"
+#include "SoundModulationValue.h"
+
+#if !UE_BUILD_SHIPPING
+#include "AudioModulationDebugger.h"
+#endif // !UE_BUILD_SHIPPING
+
+DECLARE_DWORD_COUNTER_STAT(TEXT("Bus Count"), STAT_AudioModulationBusCount, STATGROUP_AudioModulation)
+DECLARE_DWORD_COUNTER_STAT(TEXT("LFO Count"), STAT_AudioModulationLFOCount, STATGROUP_AudioModulation)
+DECLARE_DWORD_COUNTER_STAT(TEXT("Mix Count"), STAT_AudioModulationMixCount, STATGROUP_AudioModulation)
+
+
+namespace AudioModulation
+{
+	namespace
+	{
+		void MixInModulationValue(ESoundModulatorOperator& Operator, float ModStageValue, float& Value)
+		{
+			switch (Operator)
+			{
+				case ESoundModulatorOperator::Max:
+				{
+					Value = FMath::Max(ModStageValue, Value);
+				}
+				break;
+
+				case ESoundModulatorOperator::Min:
+				{
+					Value = FMath::Min(ModStageValue, Value);
+				}
+				break;
+
+				case ESoundModulatorOperator::Multiply:
+				default:
+				{
+					Value *= ModStageValue;
+					static_assert(static_cast<int32>(ESoundModulatorOperator::Count) == 3, "Possible missing operator switch case coverage");
+				}
+				break;
+			}
+		}
+	} // namespace <>
+
+	void FAudioModulationImpl::Initialize(const FAudioPluginInitializationParams& InitializationParams)
+	{
+		SourceSettings.AddDefaulted(InitializationParams.NumSources);
+	}
+
+#if WITH_EDITOR
+	void FAudioModulationImpl::OnEditSource(const USoundModulationPluginSourceSettingsBase& Settings)
+	{
+		FModulationSettingsProxy NewProxy(*CastChecked<USoundModulationSettings>(&Settings));
+		const uint32 ObjectID = Settings.GetUniqueID();
+
+		auto EditSource = [this, NewProxy, ObjectID]()
+		{
+			for (FModulationSettingsProxy& CurrentSettingsProxy : SourceSettings)
+			{
+				if (CurrentSettingsProxy.ObjectID == ObjectID)
+				{
+					CurrentSettingsProxy = NewProxy;
+				}
+			}
+		};
+
+		IsInAudioThread() ? EditSource() : FAudioThread::RunCommandOnAudioThread(EditSource);
+	}
+#endif // WITH_EDITOR
+
+	void FAudioModulationImpl::OnInitSource(const uint32 SourceId, const FName& AudioComponentUserId, const uint32 NumChannels, USoundModulationPluginSourceSettingsBase* InSettings)
+	{
+		check(IsInAudioThread());
+
+		if (USoundModulationSettings* Settings = CastChecked<USoundModulationSettings>(InSettings))
+		{
+			SourceSettings[SourceId] = FModulationSettingsProxy(*Settings);
+			for (const USoundModulatorBusMix* Mix : Settings->Mixes)
+			{
+				if (Mix)
+				{
+					ActivateBusMix(*Mix, /* bReset */ false);
+				}
+			}
+
+			auto CheckRefActivate = [this](const USoundModulatorBusBase* Bus)
+			{
+				if (!Bus)
+				{
+					return;
+				}
+
+				if (Bus->bAutoActivate)
+				{
+					ActivateBus(*Bus);
+				}
+
+				for (const USoundModulatorBase* Modulator : Bus->Modulators)
+				{
+					if (const USoundModulatorLFO* LFO = Cast<USoundModulatorLFO>(Modulator))
+					{
+						if (LFO->bAutoActivate)
+						{
+							FModulatorLFOProxy LFOProxy(*LFO);
+							ActivateLFO(*LFO);
+						}
+					}
+				}
+			};
+
+			for (const FSoundVolumeModulationInput& Input : Settings->Volume.Inputs)
+			{
+				CheckRefActivate(Cast<USoundModulatorBusBase>(Input.Bus));
+			}
+
+			for (const FSoundPitchModulationInput& Input : Settings->Pitch.Inputs)
+			{
+				CheckRefActivate(Cast<USoundModulatorBusBase>(Input.Bus));
+			}
+
+			for (const FSoundLPFModulationInput& Input : Settings->Lowpass.Inputs)
+			{
+				CheckRefActivate(Cast<USoundModulatorBusBase>(Input.Bus));
+			}
+
+			for (const FSoundHPFModulationInput& Input : Settings->Highpass.Inputs)
+			{
+				CheckRefActivate(Cast<USoundModulatorBusBase>(Input.Bus));
+			}
+		}
+	}
+
+	void FAudioModulationImpl::OnReleaseSource(const uint32 SourceId)
+	{
+		check(IsInAudioThread());
+
+		SourceSettings[SourceId] = FModulationSettingsProxy();
+	}
+
+#if !UE_BUILD_SHIPPING
+	bool FAudioModulationImpl::OnPostHelp(FCommonViewportClient* ViewportClient, const TCHAR* Stream)
+	{
+		return ViewportClient ? Debugger.OnPostHelp(*ViewportClient, Stream) : true;
+	}
+
+	int32 FAudioModulationImpl::OnRenderStat(FViewport* Viewport, FCanvas* Canvas, int32 X, int32 Y, const UFont& Font, const FVector* ViewLocation, const FRotator* ViewRotation)
+	{
+		return Canvas ? Debugger.OnRenderStat(*Canvas, X, Y, Font) : Y;
+	}
+
+	bool FAudioModulationImpl::OnToggleStat(FCommonViewportClient* ViewportClient, const TCHAR* Stream)
+	{
+		return ViewportClient ? Debugger.OnToggleStat(*ViewportClient, Stream) : true;
+	}
+#endif // !UE_BUILD_SHIPPING
+
+	void FAudioModulationImpl::ActivateBus(const USoundModulatorBusBase& Bus)
+	{
+		FModulatorBusProxy BusProxy(Bus);
+		ActivateBus(BusProxy);
+	}
+
+	void FAudioModulationImpl::ActivateBus(const FModulatorBusProxy& BusProxy)
+	{
+		auto ActivateBusInternal = [this, BusProxy]()
+		{
+			const BusId NewBusId = BusProxy.GetBusId();
+			if (!ActiveBuses.Contains(NewBusId))
+			{
+				ActiveBuses.Add(NewBusId, BusProxy);
+			}
+		};
+
+		IsInAudioThread() ? ActivateBusInternal() : FAudioThread::RunCommandOnAudioThread(ActivateBusInternal);
+	}
+
+	void FAudioModulationImpl::ActivateBusMix(const USoundModulatorBusMix& BusMix, bool bReset)
+	{
+		const BusMixId MixId = static_cast<BusMixId>(BusMix.GetUniqueID());
+		FModulatorBusMixProxy MixProxy(BusMix);
+
+		auto ActivateMixInternal = [this, MixId, MixProxy, bReset]()
+		{
+			FModulatorBusMixProxy* BusMixProxy = nullptr;
+			if (FModulatorBusMixProxy* ExistingMixProxy = ActiveBusMixes.Find(MixId))
+			{
+				// Reactivate just in case it is currently stopping but not yet stopped.
+				ExistingMixProxy->SetActive();
+				if (bReset)
+				{
+					for (int32 i = 0; i < ExistingMixProxy->Channels.Num(); ++i)
+					{
+						if (ExistingMixProxy->Channels[i].BusId == MixProxy.Channels[i].BusId)
+						{
+							ExistingMixProxy->Channels[i].Value.TargetValue = MixProxy.Channels[i].Value.TargetValue;
+						}
+					}
+				}
+
+				BusMixProxy = ExistingMixProxy;
+			}
+			else
+			{
+				BusMixProxy = &ActiveBusMixes.Emplace(MixId, MixProxy);
+			}
+		};
+
+		IsInAudioThread() ? ActivateMixInternal() : FAudioThread::RunCommandOnAudioThread(ActivateMixInternal);
+	}
+
+	void FAudioModulationImpl::ActivateLFO(const USoundModulatorLFO& LFO)
+	{
+		const AudioModulation::LFOId LFOId = static_cast<const AudioModulation::LFOId>(LFO.GetUniqueID());
+		FModulatorLFOProxy LFOProxy(LFO);
+
+		auto ActivateLFOInternal = [this, LFOId, LFOProxy]()
+		{
+			if (!ActiveLFOs.Contains(LFOId))
+			{
+				ActiveLFOs.Add(LFOId, LFOProxy);
+			}
+		};
+
+		IsInAudioThread() ? ActivateLFOInternal() : FAudioThread::RunCommandOnAudioThread(ActivateLFOInternal);
+	}
+
+	float FAudioModulationImpl::CalculateModulationValue(FModulationPatchProxy& Proxy) const
+	{
+		float Value = Proxy.DefaultInputValue;
+
+		if (!Proxy.OutputProxy.bInitialized)
+		{
+			Proxy.OutputProxy.SampleAndHoldValue = Proxy.DefaultInputValue;
+		}
+
+		for (const FModulationInputProxy& InputProxy : Proxy.InputProxies)
+		{
+			if (InputProxy.bSampleAndHold)
+			{
+				if (!Proxy.OutputProxy.bInitialized)
+				{
+					if (const FModulatorBusProxy* BusProxy = ActiveBuses.Find(InputProxy.BusId))
+					{
+						float ModStageValue = BusProxy->GetValue();
+						InputProxy.Transform.Apply(ModStageValue);
+						MixInModulationValue(Proxy.OutputProxy.Operator, ModStageValue, Proxy.OutputProxy.SampleAndHoldValue);
+					}
+				}
+			}
+			else
+			{
+				if (const FModulatorBusProxy* BusProxy = ActiveBuses.Find(InputProxy.BusId))
+				{
+					float ModStageValue = BusProxy->GetValue();
+					InputProxy.Transform.Apply(ModStageValue);
+					MixInModulationValue(Proxy.OutputProxy.Operator, ModStageValue, Value);
+				}
+			}
+		}
+
+		Proxy.OutputProxy.bInitialized = true;
+		MixInModulationValue(Proxy.OutputProxy.Operator, Proxy.OutputProxy.SampleAndHoldValue, Value);
+		Proxy.OutputProxy.Transform.Apply(Value);
+		return Value;
+	}
+
+	void FAudioModulationImpl::DeactivateBusMix(const BusMixId BusMixId)
+	{
+		auto DeactivateMix = [this, BusMixId]()
+		{
+			check(IsInAudioThread());
+
+			if (FModulatorBusMixProxy* Mix = ActiveBusMixes.Find(BusMixId))
+			{
+				Mix->SetStopping();
+			}
+		};
+
+		IsInAudioThread() ? DeactivateMix() : FAudioThread::RunCommandOnAudioThread(DeactivateMix);
+	}
+
+	void FAudioModulationImpl::DeactivateBus(const BusId BusId)
+	{
+		if (!IsInAudioThread())
+		{
+			FAudioThread::RunCommandOnAudioThread([this, BusId]()
+			{
+				DeactivateBus(BusId);
+			});
+
+			return;
+		}
+
+		ActiveBuses.Remove(BusId);
+	}
+
+	void FAudioModulationImpl::DeactivateLFO(const LFOId LFOId)
+	{
+		if (!IsInAudioThread())
+		{
+			FAudioThread::RunCommandOnAudioThread([this, LFOId]()
+			{
+				DeactivateLFO(LFOId);
+			});
+
+			return;
+		}
+
+		ActiveLFOs.Remove(LFOId);
+	}
+
+	TSet<BusId> FAudioModulationImpl::GetReferencedBusIds() const
+	{
+		auto GetActivePatchBusIds = [this](const FModulationPatchProxy& Patch, TSet<BusId>& OutBusIds)
+		{
+			for (const FModulationInputProxy& InputProxy : Patch.InputProxies)
+			{
+				if (InputProxy.bSampleAndHold)
+				{
+					if (!Patch.OutputProxy.bInitialized)
+					{
+						OutBusIds.Add(InputProxy.BusId);
+					}
+				}
+				else
+				{
+					OutBusIds.Add(InputProxy.BusId);
+				}
+			}
+		};
+
+		TSet<BusId> BusIds;
+		for (const FModulationSettingsProxy& Settings : SourceSettings)
+		{
+			GetActivePatchBusIds(Settings.Highpass, BusIds);
+			GetActivePatchBusIds(Settings.Lowpass, BusIds);
+			GetActivePatchBusIds(Settings.Pitch, BusIds);
+			GetActivePatchBusIds(Settings.Volume, BusIds);
+		}
+
+		return MoveTemp(BusIds);
+	}
+
+	bool FAudioModulationImpl::IsBusActive(const BusId BusId) const
+	{
+		check(IsInAudioThread());
+
+		return ActiveBuses.Contains(BusId);
+	}
+
+	void FAudioModulationImpl::ProcessControls(const uint32 SourceId, FSoundModulationControls& Controls)
+	{
+		check(IsInAudioThread());
+
+		FModulationSettingsProxy& Settings = SourceSettings[SourceId];
+
+		Controls.Volume   = CalculateModulationValue(Settings.Volume);
+		Controls.Pitch    = CalculateModulationValue(Settings.Pitch);
+		Controls.Lowpass  = CalculateModulationValue(Settings.Lowpass);
+		Controls.Highpass = CalculateModulationValue(Settings.Highpass);
+	}
+
+	void FAudioModulationImpl::ProcessModulators(float Elapsed)
+	{
+		check(IsInAudioThread());
+
+		// Update LFOs (prior to bus mixing to avoid single-frame latency)
+		for (TPair<LFOId, FModulatorLFOProxy>& Pair : ActiveLFOs)
+		{
+			// Clear is active as bus calls to MixLFO in subsequent loop will
+			// mark as active.  If no buses are currently using the LFO and
+			// its set to 'Limited Lifetime', then it will be marked to destroy.
+			Pair.Value.ClearIsActive();
+			Pair.Value.Update(Elapsed);
+		}
+
+		// Reset buses & refresh cached LFO
+		for (TPair<BusId, FModulatorBusProxy>& Pair : ActiveBuses)
+		{
+			Pair.Value.Reset();
+			Pair.Value.MixLFO(ActiveLFOs);
+		}
+
+		// Update mix values and apply to prescribed buses.
+		// Track bus mixes ready to remove
+		TArray<BusMixId> MixesToDeactivate;
+		for (TPair<BusMixId, FModulatorBusMixProxy>& Pair : ActiveBusMixes)
+		{
+			Pair.Value.Update(Elapsed, ActiveBuses);
+			if (Pair.Value.CanDeactivate())
+			{
+				UE_LOG(LogAudioModulation, Log, TEXT("Audio modulation mix '%u' stopped."), static_cast<uint32>(Pair.Key));
+				MixesToDeactivate.Add(Pair.Key);
+			}
+		}
+
+		{	// Destroy LFO proxies ready to deactivate
+			TArray<LFOId> LFOsToDeactivate;
+			for (TPair<LFOId, FModulatorLFOProxy>& Pair : ActiveLFOs)
+			{
+				if (Pair.Value.CanDeactivate())
+				{
+					LFOsToDeactivate.Add(Pair.Key);
+				}
+			}
+			for (const LFOId& LFOId : LFOsToDeactivate)
+			{
+				DeactivateLFO(LFOId);
+			}
+		}
+
+		{	// Destroy Bus proxies ready to deactivate
+			const TSet<BusId> ReferencedBuses = GetReferencedBusIds();
+			TArray<BusId> BusesToDeactivate;
+			for (TPair<BusId, FModulatorBusProxy>& Pair : ActiveBuses)
+			{
+				FModulatorBusProxy& BusProxy = Pair.Value;
+				if (BusProxy.CanDeactivate() && !ReferencedBuses.Contains(BusProxy.GetBusId()))
+				{
+					BusesToDeactivate.Add(Pair.Key);
+				}
+			}
+			for (const BusId& BusId : BusesToDeactivate)
+			{
+				DeactivateBus(BusId);
+			}
+		}
+
+		// Destroy mixes that have stopped
+		for (const BusMixId& MixId : MixesToDeactivate)
+		{
+			ActiveBusMixes.Remove(MixId);
+		}
+
+		SET_DWORD_STAT(STAT_AudioModulationBusCount, ActiveBuses.Num());
+		SET_DWORD_STAT(STAT_AudioModulationMixCount, ActiveBusMixes.Num());
+		SET_DWORD_STAT(STAT_AudioModulationLFOCount, ActiveLFOs.Num());
+
+#if !UE_BUILD_SHIPPING
+		Debugger.UpdateDebugData(ActiveBuses, ActiveBusMixes, ActiveLFOs);
+#endif // !UE_BUILD_SHIPPING
+	}
+
+	void FAudioModulationImpl::SetBusDefault(const USoundModulatorBusBase& Bus, const float Value)
+	{
+		auto BusId = static_cast<const AudioModulation::BusId>(Bus.GetUniqueID());
+		if (FModulatorBusProxy* BusProxy = ActiveBuses.Find(BusId))
+		{
+			BusProxy->SetDefaultValue(Value);
+		}
+	}
+
+	void FAudioModulationImpl::SetBusRange(const USoundModulatorBusBase& Bus, const FVector2D& Range)
+	{
+		auto BusId = static_cast<const AudioModulation::BusId>(Bus.GetUniqueID());
+		if (FModulatorBusProxy* BusProxy = ActiveBuses.Find(BusId))
+		{
+			BusProxy->SetRange(Range);
+		}
+	}
+
+	void FAudioModulationImpl::SetBusMixChannel(const USoundModulatorBusMix& BusMix, const USoundModulatorBusBase& Bus, const float TargetValue)
+	{
+		const AudioModulation::BusMixId MixId = static_cast<const AudioModulation::BusMixId>(BusMix.GetUniqueID());
+		const AudioModulation::BusId BusId = static_cast<const AudioModulation::BusId>(Bus.GetUniqueID());
+
+		auto SetChannel = [this, MixId, BusId, TargetValue]()
+		{
+			if (FModulatorBusMixProxy* BusMixProxy = ActiveBusMixes.Find(MixId))
+			{
+				if (FModulatorBusMixChannelProxy* ChannelProxy = BusMixProxy->Channels.Find(BusId))
+				{
+					ChannelProxy->Value.TargetValue = TargetValue;
+				}
+			}
+		};
+
+		IsInAudioThread() ? SetChannel() : FAudioThread::RunCommandOnAudioThread(SetChannel);
+	}
+
+	void FAudioModulationImpl::SetLFOFrequency(const USoundModulatorLFO& LFO, const float Freq)
+	{
+		const AudioModulation::LFOId LFOId = static_cast<const AudioModulation::LFOId>(LFO.GetUniqueID());
+
+		auto SetFreq = [this, LFOId, Freq]()
+		{
+			if (FModulatorLFOProxy* LFOProxy = ActiveLFOs.Find(LFOId))
+			{
+				LFOProxy->SetFreq(Freq);
+			}
+		};
+
+		IsInAudioThread() ? SetFreq() : FAudioThread::RunCommandOnAudioThread(SetFreq);
+	}
+} // namespace AudioModulation
+#endif // WITH_AUDIOMODULATION
diff --git a/Engine/Plugins/Runtime/AudioModulation/Source/AudioModulation/Private/AudioModulationInternal.h b/Engine/Plugins/Runtime/AudioModulation/Source/AudioModulation/Private/AudioModulationInternal.h
new file mode 100644
--- /dev/null
+++ b/Engine/Plugins/Runtime/AudioModulation/Source/AudioModulation/Private/AudioModulationInternal.h
@@ -0,0 +1,117 @@
+// Copyright 1998-2019 Epic Games, Inc. All Rights Reserved.
+#pragma once
+
+#include "CoreMinimal.h"
+
+#include "Containers/Queue.h"
+
+#include "AudioModulation.h"
+#include "SoundModulatorBase.h"
+#include "SoundModulatorBus.h"
+#include "SoundModulatorBusMix.h"
+#include "SoundModulatorLFO.h"
+
+
+#if WITH_AUDIOMODULATION
+
+#if !UE_BUILD_SHIPPING
+#include "AudioModulationDebugger.h"
+#endif // !UE_BUILD_SHIPPING
+
+namespace AudioModulation
+{
+	class FAudioModulationImpl
+	{
+	public:
+		void Initialize(const FAudioPluginInitializationParams& InitializationParams);
+
+#if WITH_EDITOR
+		void OnEditSource(const USoundModulationPluginSourceSettingsBase& Settings);
+#endif // WITH_EDITOR
+
+		void OnInitSource(const uint32 SourceId, const FName& AudioComponentUserId, const uint32 NumChannels, USoundModulationPluginSourceSettingsBase* Settings);
+		void OnReleaseSource(const uint32 SourceId);
+
+		void ActivateBus(const USoundModulatorBusBase& Bus);
+		void ActivateBus(const FModulatorBusProxy& BusProxy);
+		void ActivateBusMix(const USoundModulatorBusMix& BusMix, bool bReset);
+		void ActivateLFO(const USoundModulatorLFO& LFO);
+		void DeactivateBusMix(const BusMixId BusMixId);
+		void DeactivateBus(const BusId BusId);
+		void DeactivateLFO(const LFOId LFOId);
+
+		bool IsBusActive(const BusId BusId) const;
+
+		void ProcessAudio(const FAudioPluginSourceInputData& InputData, FAudioPluginSourceOutputData& OutputData);
+		void ProcessControls(const uint32 SourceId, FSoundModulationControls& Controls);
+		void ProcessModulators(const float Elapsed);
+
+		void SetBusDefault(const USoundModulatorBusBase& Bus, const float Value);
+		void SetBusRange(const USoundModulatorBusBase& Bus, const FVector2D& Range);
+		void SetBusMixChannel(const USoundModulatorBusMix& BusMix, const USoundModulatorBusBase& Bus, const float TargetValue);
+		void SetLFOFrequency(const USoundModulatorLFO& LFO, const float Freq);
+
+	private:
+		float CalculateModulationValue(FModulationPatchProxy& Proxy) const;
+		TSet<BusId> GetReferencedBusIds() const;
+
+		BusMixProxyMap ActiveBusMixes;
+		BusProxyMap    ActiveBuses;
+		LFOProxyMap    ActiveLFOs;
+
+		TArray<FModulationSettingsProxy> SourceSettings;
+
+#if !UE_BUILD_SHIPPING
+	public:
+		bool OnPostHelp(FCommonViewportClient* ViewportClient, const TCHAR* Stream);
+		int OnRenderStat(FViewport* Viewport, FCanvas* Canvas, int32 X, int32 Y, const UFont& Font, const FVector* ViewLocation, const FRotator* ViewRotation);
+		bool OnToggleStat(FCommonViewportClient* ViewportClient, const TCHAR* Stream);
+
+	private:
+		FAudioModulationDebugger Debugger;
+#endif // !UE_BUILD_SHIPPING
+	};
+} // namespace AudioModulation
+
+#else // WITH_AUDIOMODULATION
+
+namespace AudioModulation
+{
+	// Null implementation for compiler
+	class FAudioModulationImpl
+	{
+	public:
+		void Initialize(const FAudioPluginInitializationParams& InitializationParams) { }
+
+#if WITH_EDITOR
+		void OnEditSource(const USoundModulationPluginSourceSettingsBase& Settings) { }
+#endif // WITH_EDITOR
+
+		void OnInitSource(const uint32 SourceId, const FName& AudioComponentUserId, const uint32 NumChannels, USoundModulationPluginSourceSettingsBase* Settings) { }
+		void OnReleaseSource(const uint32 SourceId) { }
+
+#if !UE_BUILD_SHIPPING
+		bool OnPostHelp(FCommonViewportClient* ViewportClient, const TCHAR* Stream) { return false; }
+		int OnRenderStat(FViewport* Viewport, FCanvas* Canvas, int32 X, int32 Y, const UFont& Font, const FVector* ViewLocation, const FRotator* ViewRotation) { return Y; }
+		bool OnToggleStat(FCommonViewportClient* ViewportClient, const TCHAR* Stream) { return false; }
+#endif // !UE_BUILD_SHIPPING
+
+		void ActivateBus(const USoundModulatorBusBase& Bus) { }
+		void ActivateBusMix(const USoundModulatorBusMix& BusMix, bool bReset) { }
+		void ActivateLFO(const USoundModulatorLFO& LFO) { }
+		void DeactivateBusMix(const BusMixId BusMixId) { }
+		void DeactivateBus(const BusId BusId) { }
+		void DeactivateLFO(const LFOId LFOId) { }
+
+		void ProcessAudio(const FAudioPluginSourceInputData& InputData, FAudioPluginSourceOutputData& OutputData) { }
+		void ProcessControls(const uint32 SourceId, FSoundModulationControls& Controls) { }
+		void ProcessModulators(const float Elapsed) { }
+
+		void SetBusDefault(const USoundModulatorBusBase& Bus, const float Value) { }
+		void SetBusMin(const USoundModulatorBusBase& Bus, const float Value) { }
+		void SetBusMax(const USoundModulatorBusBase& Bus, const float Value) { }
+		void SetBusMixChannel(const USoundModulatorBusMix& BusMix, const USoundModulatorBusBase& Bus, const float TargetValue) { }
+		void SetLFOFrequency(const USoundModulatorLFO& LFO, const float Freq) { }
+	};
+} // namespace AudioModulation
+#endif // WITH_AUDIOMODULATION
diff --git a/Engine/Plugins/Runtime/AudioModulation/Source/AudioModulation/Private/AudioModulationLogging.cpp b/Engine/Plugins/Runtime/AudioModulation/Source/AudioModulation/Private/AudioModulationLogging.cpp
new file mode 100644
--- /dev/null
+++ b/Engine/Plugins/Runtime/AudioModulation/Source/AudioModulation/Private/AudioModulationLogging.cpp
@@ -0,0 +1,4 @@
+// Copyright 1998-2019 Epic Games, Inc. All Rights Reserved.
+#include "AudioModulationLogging.h"
+
+DEFINE_LOG_CATEGORY(LogAudioModulation);
diff --git a/Engine/Plugins/Runtime/AudioModulation/Source/AudioModulation/Private/AudioModulationLogging.h b/Engine/Plugins/Runtime/AudioModulation/Source/AudioModulation/Private/AudioModulationLogging.h
new file mode 100644
--- /dev/null
+++ b/Engine/Plugins/Runtime/AudioModulation/Source/AudioModulation/Private/AudioModulationLogging.h
@@ -0,0 +1,7 @@
+// Copyright 1998-2019 Epic Games, Inc. All Rights Reserved.
+#pragma once
+
+#include "CoreMinimal.h"
+
+
+DECLARE_LOG_CATEGORY_EXTERN(LogAudioModulation, Log, All);
diff --git a/Engine/Plugins/Runtime/AudioModulation/Source/AudioModulation/Private/AudioModulationStatics.cpp b/Engine/Plugins/Runtime/AudioModulation/Source/AudioModulation/Private/AudioModulationStatics.cpp
new file mode 100644
--- /dev/null
+++ b/Engine/Plugins/Runtime/AudioModulation/Source/AudioModulation/Private/AudioModulationStatics.cpp
@@ -0,0 +1,230 @@
+// Copyright 1998-2019 Epic Games, Inc. All Rights Reserved.
+
+#include "AudioModulationStatics.h"
+
+#include "AudioDevice.h"
+#include "AudioModulation.h"
+#include "AudioModulationInternal.h"
+#include "AudioModulationLogging.h"
+#include "AudioModulationUtils.h"
+#include "Engine/Engine.h"
+#include "SoundModulatorBus.h"
+
+#define LOCTEXT_NAMESPACE "AudioModulationStatics"
+
+
+//////////////////////////////////////////////////////////////////////////
+// UAudioModulationStatics
+
+UAudioModulationStatics::UAudioModulationStatics(const FObjectInitializer& ObjectInitializer)
+	: Super(ObjectInitializer)
+{
+}
+
+void UAudioModulationStatics::ActivateBus(const UObject* WorldContextObject, USoundModulatorBusBase* Bus)
+{
+	if (!Bus)
+	{
+		return;
+	}
+
+	UWorld* World = AudioModulation::GetAudioWorld(WorldContextObject);
+	if (AudioModulation::FAudioModulationImpl* ModulationImpl = AudioModulation::GetModulationImpl(World))
+	{
+		ModulationImpl->ActivateBus(*Bus);
+	}
+}
+
+void UAudioModulationStatics::ActivateBusMix(const UObject* WorldContextObject, USoundModulatorBusMix* BusMix)
+{
+	if (BusMix)
+	{
+		UWorld* World = AudioModulation::GetAudioWorld(WorldContextObject);
+		if (AudioModulation::FAudioModulationImpl* ModulationImpl = AudioModulation::GetModulationImpl(World))
+		{
+			ModulationImpl->ActivateBusMix(*BusMix, false);
+		}
+	}
+}
+
+void UAudioModulationStatics::ActivateModulator(const UObject* WorldContextObject, USoundModulatorBase* Modulator)
+{
+	UWorld* World = AudioModulation::GetAudioWorld(WorldContextObject);
+	if (AudioModulation::FAudioModulationImpl* ModulationImpl = AudioModulation::GetModulationImpl(World))
+	{
+		if (USoundModulatorLFO* LFO = Cast<USoundModulatorLFO>(Modulator))
+		{
+			ModulationImpl->ActivateLFO(*LFO);
+		}
+	}
+}
+
+USoundVolumeModulatorBus* UAudioModulationStatics::CreateVolumeBus(const UObject* WorldContextObject, FName Name, float DefaultValue, bool Activate)
+{
+	return AudioModulation::CreateModulatorBus<USoundVolumeModulatorBus>(WorldContextObject, Name, DefaultValue, Activate);
+}
+
+USoundPitchModulatorBus* UAudioModulationStatics::CreatePitchBus(const UObject* WorldContextObject, FName Name, float DefaultValue, bool Activate)
+{
+	return AudioModulation::CreateModulatorBus<USoundPitchModulatorBus>(WorldContextObject, Name, DefaultValue, Activate);
+}
+
+USoundLPFModulatorBus* UAudioModulationStatics::CreateLPFBus(const UObject* WorldContextObject, FName Name, float DefaultValue, bool Activate)
+{
+	return AudioModulation::CreateModulatorBus<USoundLPFModulatorBus>(WorldContextObject, Name, DefaultValue, Activate);
+}
+
+USoundHPFModulatorBus* UAudioModulationStatics::CreateHPFBus(const UObject* WorldContextObject, FName Name, float DefaultValue, bool Activate)
+{
+	return AudioModulation::CreateModulatorBus<USoundHPFModulatorBus>(WorldContextObject, Name, DefaultValue, Activate);
+}
+
+USoundModulatorLFO* UAudioModulationStatics::CreateLFO(const UObject* WorldContextObject, FName Name, float Amplitude, float Frequency, float Offset, bool Activate)
+{
+	UWorld* World = AudioModulation::GetAudioWorld(WorldContextObject);
+	if (!World)
+	{
+		return nullptr;
+	}
+
+	USoundModulatorLFO* NewLFO = NewObject<USoundModulatorLFO>(nullptr, Name, RF_Transient);
+	NewLFO->Amplitude = Amplitude;
+	NewLFO->Frequency = Frequency;
+	NewLFO->Offset    = Offset;
+
+	if (Activate)
+	{
+		if (AudioModulation::FAudioModulationImpl* ModulationImpl = AudioModulation::GetModulationImpl(World))
+		{
+			ModulationImpl->ActivateLFO(*NewLFO);
+		}
+	}
+
+	return NewLFO;
+}
+
+USoundModulatorBusMix* UAudioModulationStatics::CreateBusMix(const UObject* WorldContextObject, FName Name, TArray<USoundModulatorBusBase*> Buses, float TargetValue, bool Activate)
+{
+	UWorld* World = AudioModulation::GetAudioWorld(WorldContextObject);
+	if (!World)
+	{
+		return nullptr;
+	}
+
+	USoundModulatorBusMix* NewBusMix = NewObject<USoundModulatorBusMix>(nullptr, Name, RF_Transient);
+	for (USoundModulatorBusBase* Bus : Buses)
+	{
+		if (Bus)
+		{
+			NewBusMix->Channels.Emplace_GetRef(Bus, TargetValue);
+		}
+		else
+		{
+			UE_LOG(LogAudioModulation, Warning,
+				TEXT("USoundModulatorBusMix '%s' was created but bus provided is null. Channel not added."),
+				*Name.ToString());
+		}
+	}
+
+	if (Activate)
+	{
+		if (AudioModulation::FAudioModulationImpl* ModulationImpl = AudioModulation::GetModulationImpl(World))
+		{
+			ModulationImpl->ActivateBusMix(*NewBusMix, false);
+		}
+	}
+
+	return NewBusMix;
+}
+
+void UAudioModulationStatics::DeactivateBus(const UObject* WorldContextObject, USoundModulatorBusBase* Bus)
+{
+	if (!Bus)
+	{
+		return;
+	}
+
+	UWorld* World = AudioModulation::GetAudioWorld(WorldContextObject);
+	if (AudioModulation::FAudioModulationImpl* ModulationImpl = AudioModulation::GetModulationImpl(World))
+	{
+		auto BusId = static_cast<const AudioModulation::BusId>(Bus->GetUniqueID());
+		ModulationImpl->DeactivateBus(BusId);
+	}
+}
+
+void UAudioModulationStatics::DeactivateBusMix(const UObject* WorldContextObject, USoundModulatorBusMix* BusMix)
+{
+	if (BusMix)
+	{
+		UWorld* World = AudioModulation::GetAudioWorld(WorldContextObject);
+		if (AudioModulation::FAudioModulationImpl* ModulationImpl = AudioModulation::GetModulationImpl(World))
+		{
+			auto BusMixId = static_cast<const AudioModulation::BusMixId>(BusMix->GetUniqueID());
+			ModulationImpl->DeactivateBusMix(BusMixId);
+		}
+	}
+}
+
+void UAudioModulationStatics::DeactivateModulator(const UObject* WorldContextObject, USoundModulatorBase* Modulator)
+{
+	UWorld* World = AudioModulation::GetAudioWorld(WorldContextObject);
+	if (AudioModulation::FAudioModulationImpl* ModulationImpl = AudioModulation::GetModulationImpl(World))
+	{
+		if (USoundModulatorLFO* LFO = Cast<USoundModulatorLFO>(Modulator))
+		{
+			auto LFOId = static_cast<const AudioModulation::LFOId>(Modulator->GetUniqueID());
+			ModulationImpl->DeactivateLFO(LFOId);
+		}
+	}
+}
+
+void UAudioModulationStatics::SetBusDefault(const UObject* WorldContextObject, USoundModulatorBusBase* Bus, float Value)
+{
+	if (Bus)
+	{
+		UWorld* World = AudioModulation::GetAudioWorld(WorldContextObject);
+		if (AudioModulation::FAudioModulationImpl* ModulationImpl = AudioModulation::GetModulationImpl(World))
+		{
+			ModulationImpl->SetBusDefault(*Bus, Value);
+		}
+	}
+	else
+	{
+		UE_LOG(LogAudioModulation, Warning, TEXT("Bus not specified. Bus default set request ignored."));
+	}
+}
+
+void UAudioModulationStatics::SetBusMixChannel(const UObject* WorldContextObject, USoundModulatorBusMix* BusMix, USoundModulatorBusBase* Bus, float TargetValue)
+{
+	if (BusMix && Bus)
+	{
+		UWorld* World = AudioModulation::GetAudioWorld(WorldContextObject);
+		if (AudioModulation::FAudioModulationImpl* ModulationImpl = AudioModulation::GetModulationImpl(World))
+		{
+			ModulationImpl->SetBusMixChannel(*BusMix, *Bus, TargetValue);
+		}
+	}
+	else
+	{
+		UE_LOG(LogAudioModulation, Warning, TEXT("USoundModulatorBusMix or bus not specified. Bus channel mix set request ignored."));
+	}
+}
+
+void UAudioModulationStatics::SetLFOFrequency(const UObject* WorldContextObject, USoundModulatorLFO* LFO, float Freq)
+{
+	if (LFO)
+	{
+		UWorld* World = AudioModulation::GetAudioWorld(WorldContextObject);
+		if (AudioModulation::FAudioModulationImpl* ModulationImpl = AudioModulation::GetModulationImpl(World))
+		{
+			ModulationImpl->SetLFOFrequency(*LFO, Freq);
+		}
+	}
+	else
+	{
+		UE_LOG(LogAudioModulation, Warning,
+			TEXT("USoundModulatorLFO not specified. LFO frequency set ignored."));
+	}
+}
+
+#undef LOCTEXT_NAMESPACE
diff --git a/Engine/Plugins/Runtime/AudioModulation/Source/AudioModulation/Private/AudioModulationStatics.h b/Engine/Plugins/Runtime/AudioModulation/Source/AudioModulation/Private/AudioModulationStatics.h
new file mode 100644
--- /dev/null
+++ b/Engine/Plugins/Runtime/AudioModulation/Source/AudioModulation/Private/AudioModulationStatics.h
@@ -0,0 +1,131 @@
+// Copyright 1998-2019 Epic Games, Inc. All Rights Reserved.
+#pragma once
+
+#include "CoreMinimal.h"
+#include "Kismet/BlueprintFunctionLibrary.h"
+#include "UObject/Object.h"
+#include "UObject/ObjectMacros.h"
+
+#include "AudioModulationStatics.generated.h"
+
+
+class USoundModulatorBus;
+class USoundModulatorBusMix;
+
+UCLASS()
+class UAudioModulationStatics : public UBlueprintFunctionLibrary
+{
+	GENERATED_UCLASS_BODY()
+
+public:
+	/** Activates a bus. Does nothing if an instance of the provided bus is already active
+	 * @param Bus - Bus to activate
+	 */
+	UFUNCTION(BlueprintCallable, Category = "Audio", DisplayName = "Activate Sound Bus", meta = (WorldContext = "WorldContextObject", Keywords = "activate modulation modulator bus"))
+	static void ActivateBus(const UObject* WorldContextObject, USoundModulatorBusBase* Bus);
+
+	/** Activates a bus modulator mix. Does nothing if an instance of the provided bus mix is already active
+	 * @param BusMix - Mix to activate
+	 */
+	UFUNCTION(BlueprintCallable, Category = "Audio", DisplayName = "Activate Sound Modulation Bus Mix", meta = (WorldContext = "WorldContextObject", Keywords = "activate modulation modulator bus mix"))
+	static void ActivateBusMix(const UObject* WorldContextObject, USoundModulatorBusMix* BusMix);
+
+	/** Activates a modulator. Does nothing if an instance of the provided modulator is already active
+	 * @param Modulator - Modulator to activate
+	 */
+	UFUNCTION(BlueprintCallable, Category = "Audio", DisplayName = "Activate Sound Modulator", meta = (WorldContext = "WorldContextObject", Keywords = "activate modulation modulator"))
+	static void ActivateModulator(const UObject* WorldContextObject, USoundModulatorBase* Modulator);
+
+	/** Creates a volume modulation bus with the provided default value.
+	 * @param Name - Name of bus
+	 * @param DefaultValue - Default value for created bus
+	 * @param Activate - Whether or not to activate bus on creation.
+	 */
+	UFUNCTION(BlueprintCallable, Category = "Audio", DisplayName = "Create Volume Bus", meta = (WorldContext = "WorldContextObject", Keywords = "create bus modulation volume modulator"))
+	static USoundVolumeModulatorBus* CreateVolumeBus(const UObject* WorldContextObject, FName Name, float DefaultValue, bool Activate);
+
+	/** Creates a pitch modulation bus with the provided default value.
+	 * @param Name - Name of bus
+	 * @param DefaultValue - Default value for created bus
+	 * @param Activate - Whether or not to activate bus on creation.
+	 */
+	UFUNCTION(BlueprintCallable, Category = "Audio", DisplayName = "Create Sound Modulator Pitch Bus", meta = (WorldContext = "WorldContextObject", Keywords = "create bus modulation pitch modulator"))
+	static USoundPitchModulatorBus* CreatePitchBus(const UObject* WorldContextObject, FName Name, float DefaultValue, bool Activate);
+
+	/** Creates a high-pass filter (HPF) modulation bus with the provided default value.
+	 * @param Name - Name of bus
+	 * @param DefaultValue - Default value for created bus
+	 * @param Activate - Whether or not to activate bus on creation.
+	 */
+	UFUNCTION(BlueprintCallable, Category = "Audio", DisplayName = "Create Sound Modulator High-Pass Bus", meta = (WorldContext = "WorldContextObject", Keywords = "create bus modulation HPF modulator"))
+	static USoundHPFModulatorBus* CreateHPFBus(const UObject* WorldContextObject, FName Name, float DefaultValue, bool Activate);
+
+	/** Creates a low-pass filter (LPF) modulation bus with the provided default value.
+	 * @param Name - Name of bus
+	 * @param DefaultValue - Default value for created bus
+	 * @param Activate - Whether or not to activate bus on creation.
+	 */
+	UFUNCTION(BlueprintCallable, Category = "Audio", DisplayName = "Create Sound Modulator Low-Pass Bus", meta = (WorldContext = "WorldContextObject", Keywords = "create bus modulation LPF modulator"))
+	static USoundLPFModulatorBus* CreateLPFBus(const UObject* WorldContextObject, FName Name, float DefaultValue, bool Activate);
+
+	/** Creates a modulation LFO.
+	 * @param Name - Name of LFO
+	 * @param Amplitude - Amplitude of new LFO.
+	 * @param Frequency - Frequency of new LFO.
+	 * @param Offset - Offset of new LFO.
+	 * @param Activate - Whether or not to activate lfo on creation.
+	 */
+	UFUNCTION(BlueprintCallable, Category = "Audio", DisplayName = "Create Sound Modulation LFO", meta = (WorldContext = "WorldContextObject", Keywords = "create lfo modulation modulator"))
+	static USoundModulatorLFO* CreateLFO(const UObject* WorldContextObject, FName Name, float Amplitude, float Frequency, float Offset, bool Activate);
+
+	/** Creates a modulation bus mix and adds a bus channel set to the provided target value
+	 * @param Bus - Bus to add to mix
+	 * @param TargetValue - value for added bus channel to target
+	 * @param Activate - Whether or not to activate mix on creation.
+	 */
+	UFUNCTION(BlueprintCallable, Category = "Audio", DisplayName = "Create Sound Modulation Bus Mix", meta = (WorldContext = "WorldContextObject", Keywords = "create bus mix modulation modulator"))
+	static USoundModulatorBusMix* CreateBusMix(const UObject* WorldContextObject, FName Name, TArray<USoundModulatorBusBase*> Buses, float TargetValue, bool Activate);
+
+	/** Deactivates a bus. Does nothing if an instance of the provided bus is already inactive
+	 * @param Bus - Scope of modulator
+	 */
+	UFUNCTION(BlueprintCallable, Category = "Audio", DisplayName = "Deactivate Sound Modulator Bus", meta = (WorldContext = "WorldContextObject", Keywords = "deactivate modulation modulator bus"))
+	static void DeactivateBus(const UObject* WorldContextObject, USoundModulatorBusBase* Bus);
+
+	/** Deactivates a modulation bus mix. Does nothing if an instance of the provided bus mix is already inactive
+	 * @param BusMix - Mix to deactivate
+	 */
+	UFUNCTION(BlueprintCallable, Category = "Audio", DisplayName = "Deactivate Sound Modulation Bus Mix", meta = (WorldContext = "WorldContextObject", Keywords = "deactivate modulation modulator"))
+	static void DeactivateBusMix(const UObject* WorldContextObject, USoundModulatorBusMix* BusMix);
+
+	/** Deactivates a modulator. Does nothing if an instance of the provided bus mix is already inactive
+	 * @param Modulator - Modulator to activate
+	 * @param Scope - Scope of modulator
+	 */
+	UFUNCTION(BlueprintCallable, Category = "Audio", DisplayName = "Deactivate Sound Modulator", meta = (WorldContext = "WorldContextObject", Keywords = "deactivate modulation modulator"))
+	static void DeactivateModulator(const UObject* WorldContextObject, USoundModulatorBase* Modulator);
+
+	/** Sets a bus's default value
+	 * @param Bus - Bus of which to modify default value
+	 * @param Value - Value for channel to target
+	 */
+	UFUNCTION(BlueprintCallable, Category = "Audio", DisplayName = "Set Bus Default Value", meta = (WorldContext = "WorldContextObject", Keywords = "set bus default modulation modulator"))
+	static void SetBusDefault(const UObject* WorldContextObject, USoundModulatorBusBase* Bus, float Value);
+
+	/** Sets a mix's bus channel(s) to the provided target value
+	 * @param BusMix - BusMix to adjust channel of
+	 * @param Bus - Bus of channel to modify the target value
+	 * @param TargetValue - Value for channel to target
+	 */
+	UFUNCTION(BlueprintCallable, Category = "Audio", DisplayName = "Set Bus Mix Channel", meta = (WorldContext = "WorldContextObject", Keywords = "set bus mix channel modulation modulator"))
+	static void SetBusMixChannel(const UObject* WorldContextObject, USoundModulatorBusMix* BusMix, USoundModulatorBusBase* Bus, float TargetValue);
+
+
+	/** Sets an LFO's frequency to the provided value
+	 * @param LFO - LFO to adjust frequency
+	 * @param Bus - Bus of channel to modify the target value
+	 * @param TargetValue - Value for channel to target
+	 */
+	UFUNCTION(BlueprintCallable, Category = "Audio", DisplayName = "Set Modulation LFO Frequency", meta = (WorldContext = "WorldContextObject", Keywords = "set lfo modulation modulator"))
+	static void SetLFOFrequency(const UObject* WorldContextObject, USoundModulatorLFO* LFO, float Freq);
+};
diff --git a/Engine/Plugins/Runtime/AudioModulation/Source/AudioModulation/Private/AudioModulationUtils.h b/Engine/Plugins/Runtime/AudioModulation/Source/AudioModulation/Private/AudioModulationUtils.h
new file mode 100644
--- /dev/null
+++ b/Engine/Plugins/Runtime/AudioModulation/Source/AudioModulation/Private/AudioModulationUtils.h
@@ -0,0 +1,88 @@
+// Copyright 1998-2019 Epic Games, Inc. All Rights Reserved.
+#pragma once
+
+#include "CoreMinimal.h"
+#include "UObject/ObjectMacros.h"
+#include "UObject/Object.h"
+
+#include "AudioModulation.h"
+#include "AudioModulationInternal.h"
+
+
+namespace AudioModulation
+{
+	UWorld* GetAudioWorld(const UObject* WorldContextObject)
+	{
+		if (!GEngine || !GEngine->UseSound())
+		{
+			return nullptr;
+		}
+
+		UWorld* World = GEngine->GetWorldFromContextObject(WorldContextObject, EGetWorldErrorMode::LogAndReturnNull);
+		if (!World || !World->bAllowAudioPlayback || World->IsNetMode(NM_DedicatedServer))
+		{
+			return nullptr;
+		}
+
+		return World;
+	}
+
+	FAudioModulationImpl* GetModulationImpl(FAudioDevice* AudioDevice)
+	{
+		if (AudioDevice)
+		{
+			if (AudioDevice->IsModulationPluginEnabled())
+			{
+				if (IAudioModulation* ModulationInterface = AudioDevice->ModulationInterface.Get())
+				{
+					return static_cast<AudioModulation::FAudioModulation*>(ModulationInterface)->GetImpl();
+				}
+			}
+		}
+
+		return nullptr;
+	}
+
+	FAudioModulationImpl* GetModulationImpl(UWorld* World)
+	{
+		FAudioDevice* AudioDevice = nullptr;
+		if (World)
+		{
+			AudioDevice = World->GetAudioDevice();
+		}
+		else
+		{
+			if (GEngine)
+			{
+				if (FAudioDeviceManager* DeviceManager = GEngine->GetAudioDeviceManager())
+				{
+					AudioDevice = DeviceManager->GetActiveAudioDevice();
+				}
+			}
+		}
+		return GetModulationImpl(AudioDevice);
+	}
+
+	template <class T>
+	T* CreateModulatorBus(const UObject* WorldContextObject, FName Name, float DefaultValue, bool Activate)
+	{
+		UWorld* World = GetAudioWorld(WorldContextObject);
+		if (!World)
+		{
+			return nullptr;
+		}
+
+		T* NewBus = NewObject<T>(nullptr, Name, RF_Transient);
+		NewBus->DefaultValue = DefaultValue;
+
+		if (Activate)
+		{
+			if (AudioModulation::FAudioModulationImpl* ModulationImpl = GetModulationImpl(World))
+			{
+				ModulationImpl->ActivateBus(*NewBus);
+			}
+		}
+
+		return NewBus;
+	}
+} // namespace AudioModulation
\ No newline at end of file
diff --git a/Engine/Plugins/Runtime/AudioModulation/Source/AudioModulation/Private/SoundModulationPatch.cpp b/Engine/Plugins/Runtime/AudioModulation/Source/AudioModulation/Private/SoundModulationPatch.cpp
new file mode 100644
--- /dev/null
+++ b/Engine/Plugins/Runtime/AudioModulation/Source/AudioModulation/Private/SoundModulationPatch.cpp
@@ -0,0 +1,167 @@
+// Copyright 1998-2019 Epic Games, Inc. All Rights Reserved.
+#include "SoundModulationPatch.h"
+
+#include "AudioDevice.h"
+#include "AudioModulation.h"
+#include "AudioModulationInternal.h"
+#include "AudioModulationUtils.h"
+#include "SoundModulationTransform.h"
+
+
+USoundModulationSettings::USoundModulationSettings(const FObjectInitializer& ObjectInitializer)
+	: Super(ObjectInitializer)
+{
+	Lowpass.Output.Transform.OutputMin = 20.0f;
+	Lowpass.Output.Transform.OutputMax = 20000.0f;
+	Lowpass.Output.Operator = ESoundModulatorOperator::Min;
+
+	Highpass.DefaultInputValue = 0.0f;
+	Highpass.Output.Transform.OutputMin = 20.0f;
+	Highpass.Output.Transform.OutputMax = 20000.0f;
+	Highpass.Output.Operator = ESoundModulatorOperator::Max;
+}
+
+#if WITH_EDITOR
+void USoundModulationSettings::PostEditChangeProperty(FPropertyChangedEvent& PropertyChangedEvent)
+{
+	if (!PropertyChangedEvent.Property)
+	{
+		return;
+	}
+
+	if (AudioModulation::FAudioModulationImpl* Impl = AudioModulation::GetModulationImpl(GetWorld()))
+	{
+		Impl->OnEditSource(*this);
+	}
+}
+
+void USoundModulationSettings::PostEditChangeChainProperty(FPropertyChangedChainEvent& PropertyChangedEvent)
+{
+	if (!PropertyChangedEvent.Property)
+	{
+		return;
+	}
+
+	if (AudioModulation::FAudioModulationImpl* Impl = AudioModulation::GetModulationImpl(GetWorld()))
+	{
+		Impl->OnEditSource(*this);
+	}
+}
+#endif // WITH_EDITOR
+
+FSoundModulationOutput::FSoundModulationOutput()
+	: Operator(ESoundModulatorOperator::Multiply)
+{
+}
+
+FSoundModulationInputBase::FSoundModulationInputBase()
+	: bSampleAndHold(0)
+{
+}
+
+FSoundVolumeModulationInput::FSoundVolumeModulationInput()
+	: Bus(nullptr)
+{
+}
+
+FSoundPitchModulationInput::FSoundPitchModulationInput()
+	: Bus(nullptr)
+{
+}
+
+FSoundHPFModulationInput::FSoundHPFModulationInput()
+	: Bus(nullptr)
+{
+}
+
+FSoundLPFModulationInput::FSoundLPFModulationInput()
+	: Bus(nullptr)
+{
+}
+
+FSoundModulationPatchBase::FSoundModulationPatchBase()
+	: DefaultInputValue(1.0f)
+{
+}
+
+namespace AudioModulation
+{
+	FModulationInputProxy::FModulationInputProxy()
+		: BusId(InvalidBusId)
+		, bSampleAndHold(0)
+	{
+	}
+
+	FModulationInputProxy::FModulationInputProxy(const FSoundModulationInputBase& Input)
+		: BusId(InvalidBusId)
+		, Transform(Input.Transform)
+		, bSampleAndHold(Input.bSampleAndHold)
+	{
+		if (const USoundModulatorBusBase* Bus = Input.GetBus())
+		{
+			BusId = static_cast<AudioModulation::BusId>(Bus->GetUniqueID());
+		}
+	}
+
+	FModulationOutputProxy::FModulationOutputProxy()
+		: bInitialized(0)
+		, Operator(ESoundModulatorOperator::Multiply)
+		, SampleAndHoldValue(1.0f)
+	{
+	}
+
+	FModulationOutputProxy::FModulationOutputProxy(const FSoundModulationOutput& Output)
+		: bInitialized(0)
+		, Operator(Output.Operator)
+		, SampleAndHoldValue(1.0f)
+		, Transform(Output.Transform)
+	{
+	}
+
+	FModulationPatchProxy::FModulationPatchProxy()
+		: DefaultInputValue(1.0f)
+	{
+	}
+
+	FModulationPatchProxy::FModulationPatchProxy(const FSoundModulationPatchBase& Patch)
+		: DefaultInputValue(Patch.DefaultInputValue)
+		, OutputProxy(Patch.Output)
+	{
+		Patch.GenerateProxies(InputProxies);
+	}
+
+	FModulationSettingsProxy::FModulationSettingsProxy()
+#if WITH_EDITOR
+		: ObjectID(0)
+#endif // WITH_EDITOR
+	{
+		Lowpass.OutputProxy.Transform.OutputMin  = 20.0f;
+		Lowpass.OutputProxy.Transform.OutputMax  = 20000.0f;
+		Lowpass.OutputProxy.SampleAndHoldValue   = 20000.0f;
+		Lowpass.OutputProxy.Operator = ESoundModulatorOperator::Min;
+
+		Highpass.DefaultInputValue = 0.0f;
+		Highpass.OutputProxy.Transform.OutputMin = 20.0f;
+		Highpass.OutputProxy.Transform.OutputMax = 20000.0f;
+		Highpass.OutputProxy.SampleAndHoldValue  = 20.0f;
+		Highpass.OutputProxy.Operator = ESoundModulatorOperator::Max;
+	}
+
+	FModulationSettingsProxy::FModulationSettingsProxy(const USoundModulationSettings& Settings)
+		: Volume(Settings.Volume)
+		, Pitch(Settings.Pitch)
+		, Lowpass(Settings.Lowpass)
+		, Highpass(Settings.Highpass)
+#if WITH_EDITOR
+		, ObjectID(Settings.GetUniqueID())
+#endif // WITH_EDITOR
+	{
+		for (const USoundModulatorBusMix* Mix : Settings.Mixes)
+		{
+			if (Mix)
+			{
+				Mixes.Add(static_cast<BusMixId>(Mix->GetUniqueID()));
+			}
+		}
+	}
+} // namespace AudioModulation
\ No newline at end of file
diff --git a/Engine/Plugins/Runtime/AudioModulation/Source/AudioModulation/Private/SoundModulationPatch.h b/Engine/Plugins/Runtime/AudioModulation/Source/AudioModulation/Private/SoundModulationPatch.h
new file mode 100644
--- /dev/null
+++ b/Engine/Plugins/Runtime/AudioModulation/Source/AudioModulation/Private/SoundModulationPatch.h
@@ -0,0 +1,300 @@
+// Copyright 1998-2019 Epic Games, Inc. All Rights Reserved.
+#pragma once
+
+#include "CoreMinimal.h"
+#include "UObject/ObjectMacros.h"
+#include "UObject/Object.h"
+#include "SoundModulationTransform.h"
+#include "SoundModulationValue.h"
+#include "SoundModulatorBus.h"
+#include "SoundModulatorBusMix.h"
+
+#include "SoundModulationPatch.generated.h"
+
+
+// Forward Declarations
+namespace AudioModulation
+{
+	struct FModulationInputProxy;
+} // namespace AudioModulation
+
+
+USTRUCT(BlueprintType)
+struct FSoundModulationOutput
+{
+	GENERATED_USTRUCT_BODY()
+
+	FSoundModulationOutput();
+
+	/** Operator used when mixing input values */
+	UPROPERTY(EditAnywhere, Category = Modulation, BlueprintReadWrite)
+	ESoundModulatorOperator Operator;
+
+	/** Final transform before passing to output */
+	UPROPERTY(EditAnywhere, Category = Modulation, BlueprintReadWrite, meta = (ShowOnlyInnerProperties))
+	FSoundModulationOutputTransform Transform;
+};
+
+USTRUCT(BlueprintType)
+struct FSoundModulationInputBase
+{
+	GENERATED_USTRUCT_BODY()
+
+	FSoundModulationInputBase();
+	virtual ~FSoundModulationInputBase() = default;
+
+	/** Get the modulated input value on parent patch initialization and hold that value for its lifetime */
+	UPROPERTY(EditAnywhere, BlueprintReadWrite, Category = Input, meta = (DisplayName = "Sample-And-Hold"))
+	uint8 bSampleAndHold : 1;
+
+	/** Transform to apply to the input prior to mix phase */
+	UPROPERTY(EditAnywhere, BlueprintReadWrite, Category = Input, meta = (ShowOnlyInnerProperties))
+	FSoundModulationInputTransform Transform;
+
+	virtual const USoundModulatorBusBase* GetBus() const PURE_VIRTUAL(FSoundModulationInputBase::GetBus, return nullptr; );
+};
+
+USTRUCT(BlueprintType)
+struct FSoundVolumeModulationInput : public FSoundModulationInputBase
+{
+	GENERATED_USTRUCT_BODY()
+
+	FSoundVolumeModulationInput();
+
+	/** The input bus */
+	UPROPERTY(EditAnywhere, BlueprintReadWrite, Category = Input)
+	USoundVolumeModulatorBus* Bus;
+
+	virtual const USoundModulatorBusBase* GetBus() const { return Cast<USoundModulatorBusBase>(Bus); }
+};
+
+USTRUCT(BlueprintType)
+struct FSoundPitchModulationInput : public FSoundModulationInputBase
+{
+	GENERATED_USTRUCT_BODY()
+
+	FSoundPitchModulationInput();
+
+	/** The input bus */
+	UPROPERTY(EditAnywhere, BlueprintReadWrite, Category = Input)
+	USoundPitchModulatorBus* Bus;
+
+	virtual const USoundModulatorBusBase* GetBus() const { return Cast<USoundModulatorBusBase>(Bus); }
+};
+
+USTRUCT(BlueprintType)
+struct FSoundLPFModulationInput : public FSoundModulationInputBase
+{
+	GENERATED_USTRUCT_BODY()
+
+	FSoundLPFModulationInput();
+
+	/** The input bus */
+	UPROPERTY(EditAnywhere, BlueprintReadWrite, Category = Input)
+	USoundLPFModulatorBus* Bus;
+
+	virtual const USoundModulatorBusBase* GetBus() const { return Cast<USoundModulatorBusBase>(Bus); }
+};
+
+USTRUCT(BlueprintType)
+struct FSoundHPFModulationInput : public FSoundModulationInputBase
+{
+	GENERATED_USTRUCT_BODY()
+
+	FSoundHPFModulationInput();
+
+	/** The input bus */
+	UPROPERTY(EditAnywhere, BlueprintReadWrite, Category = Input)
+	USoundHPFModulatorBus* Bus;
+
+	virtual const USoundModulatorBusBase* GetBus() const { return Cast<USoundModulatorBusBase>(Bus); }
+};
+
+USTRUCT(BlueprintType)
+struct FSoundModulationPatchBase
+{
+	GENERATED_USTRUCT_BODY()
+
+	FSoundModulationPatchBase();
+
+	virtual ~FSoundModulationPatchBase() = default;
+
+	/** Default value of patch (value used when all inputs are either not provided or not active) */
+	UPROPERTY(EditAnywhere, BlueprintReadWrite, Category = Inputs)
+	float DefaultInputValue;
+
+	/** Final modulation parameters to apply */
+	UPROPERTY(EditAnywhere, BlueprintReadWrite, Category = Output, meta = (ShowOnlyInnerProperties))
+	FSoundModulationOutput Output;
+
+	virtual void GenerateProxies(TArray<AudioModulation::FModulationInputProxy>& InputProxies) const PURE_VIRTUAL(FSoundModulationPatchBase::GenerateProxies, );
+};
+
+USTRUCT(BlueprintType)
+struct FSoundVolumeModulationPatch : public FSoundModulationPatchBase
+{
+	GENERATED_USTRUCT_BODY()
+
+	/** Modulation inputs */
+	UPROPERTY(EditAnywhere, BlueprintReadWrite, Category = Inputs)
+	TArray<FSoundVolumeModulationInput> Inputs;
+
+	virtual void GenerateProxies(TArray<AudioModulation::FModulationInputProxy>& InputProxies) const
+	{
+		for (const FSoundVolumeModulationInput& Input : Inputs)
+		{
+			InputProxies.Emplace_GetRef(Input);
+		}
+	}
+};
+
+USTRUCT(BlueprintType)
+struct FSoundPitchModulationPatch : public FSoundModulationPatchBase
+{
+	GENERATED_USTRUCT_BODY()
+
+	/** Modulation inputs */
+	UPROPERTY(EditAnywhere, BlueprintReadWrite, Category = Inputs)
+	TArray<FSoundPitchModulationInput> Inputs;
+
+	virtual void GenerateProxies(TArray<AudioModulation::FModulationInputProxy>& InputProxies) const
+	{
+		for (const FSoundPitchModulationInput& Input : Inputs)
+		{
+			InputProxies.Emplace_GetRef(Input);
+		}
+	}
+};
+
+USTRUCT(BlueprintType)
+struct FSoundLPFModulationPatch : public FSoundModulationPatchBase
+{
+	GENERATED_USTRUCT_BODY()
+
+	/** Modulation inputs */
+	UPROPERTY(EditAnywhere, BlueprintReadWrite, Category = Inputs)
+	TArray<FSoundLPFModulationInput> Inputs;
+
+	virtual void GenerateProxies(TArray<AudioModulation::FModulationInputProxy>& InputProxies) const
+	{
+		for (const FSoundLPFModulationInput& Input : Inputs)
+		{
+			InputProxies.Emplace_GetRef(Input);
+		}
+	}
+};
+
+USTRUCT(BlueprintType)
+struct FSoundHPFModulationPatch : public FSoundModulationPatchBase
+{
+	GENERATED_USTRUCT_BODY()
+
+	/** Modulation inputs */
+	UPROPERTY(EditAnywhere, BlueprintReadWrite, Category = Inputs)
+	TArray<FSoundHPFModulationInput> Inputs;
+
+	virtual void GenerateProxies(TArray<AudioModulation::FModulationInputProxy>& InputProxies) const
+	{
+		for (const FSoundHPFModulationInput& Input : Inputs)
+		{
+			InputProxies.Emplace_GetRef(Input);
+		}
+	}
+};
+
+UCLASS(config = Engine, editinlinenew, BlueprintType, MinimalAPI)
+class USoundModulationSettings : public USoundModulationPluginSourceSettingsBase
+{
+	GENERATED_UCLASS_BODY()
+
+public:
+	UPROPERTY(EditAnywhere, BlueprintReadWrite, Category = Volume)
+	FSoundVolumeModulationPatch Volume;
+
+	UPROPERTY(EditAnywhere, BlueprintReadWrite, Category = Pitch)
+	FSoundPitchModulationPatch Pitch;
+
+	UPROPERTY(EditAnywhere, BlueprintReadWrite, Category = Lowpass)
+	FSoundLPFModulationPatch Lowpass;
+
+	UPROPERTY(EditAnywhere, BlueprintReadWrite, Category = Highpass)
+	FSoundHPFModulationPatch Highpass;
+
+	// Mixes that will applied and removed when sounds utilizing settings
+	// play and stop respectively. If mix has already been applied manually,
+	// mix will be removed once all sound settings referencing mix stop. Manual
+	// mix activation is ignored if already activated by means of modulation settings.
+	UPROPERTY(EditAnywhere, BlueprintReadWrite, Category = Mixes)
+	TArray<USoundModulatorBusMix*> Mixes;
+
+#if WITH_EDITOR
+	virtual void PostEditChangeProperty(FPropertyChangedEvent& PropertyChangedEvent) override;
+	virtual void PostEditChangeChainProperty(FPropertyChangedChainEvent& PropertyChangedEvent) override;
+#endif
+};
+
+
+namespace AudioModulation
+{
+	/** Modulation input instance */
+	struct FModulationInputProxy
+	{
+		FModulationInputProxy();
+		FModulationInputProxy(const FSoundModulationInputBase& Patch);
+
+		BusId BusId;
+		FSoundModulationInputTransform Transform;
+		uint8 bSampleAndHold : 1;
+	};
+
+	/** Patch applied as the final stage of a modulation chain prior to output on the sound level (Always active, never removed) */
+	struct FModulationOutputProxy
+	{
+		FModulationOutputProxy();
+		FModulationOutputProxy(const FSoundModulationOutput& Patch);
+
+		/** Whether patch has been initialized or not */
+		uint8 bInitialized : 1;
+
+		/** Operator used to calculate the output proxy value */
+		ESoundModulatorOperator Operator;
+
+		/** Cached value of sample-and-hold input values */
+		float SampleAndHoldValue;
+
+		/** Final transform before passing to output */
+		FSoundModulationOutputTransform Transform;
+	};
+
+	struct FModulationPatchProxy
+	{
+		FModulationPatchProxy();
+		FModulationPatchProxy(const FSoundModulationPatchBase& Patch);
+
+		/** Default value of patch (value used when inputs are either not provided or not active)*/
+		float DefaultInputValue;
+
+		/** Optional modulation inputs */
+		TArray<FModulationInputProxy> InputProxies;
+
+		/** Final output modulation post input combination */
+		FModulationOutputProxy OutputProxy;
+	};
+
+	struct FModulationSettingsProxy
+	{
+		FModulationSettingsProxy();
+		FModulationSettingsProxy(const USoundModulationSettings& Settings);
+
+		FModulationPatchProxy Volume;
+		FModulationPatchProxy Pitch;
+		FModulationPatchProxy Lowpass;
+		FModulationPatchProxy Highpass;
+
+		TArray<BusMixId> Mixes;
+
+#if WITH_EDITOR
+		uint32 ObjectID;
+#endif // WITH_EDITOR
+	};
+} // namespace AudioModulation
diff --git a/Engine/Plugins/Runtime/AudioModulation/Source/AudioModulation/Private/SoundModulationTransform.cpp b/Engine/Plugins/Runtime/AudioModulation/Source/AudioModulation/Private/SoundModulationTransform.cpp
new file mode 100644
--- /dev/null
+++ b/Engine/Plugins/Runtime/AudioModulation/Source/AudioModulation/Private/SoundModulationTransform.cpp
@@ -0,0 +1,113 @@
+// Copyright 1998-2019 Epic Games, Inc. All Rights Reserved.
+
+#include "SoundModulationTransform.h"
+
+#include "Audio.h"
+
+
+FSoundModulationInputTransform::FSoundModulationInputTransform()
+	: InputMin(0.0f)
+	, InputMax(1.0f)
+	, OutputMin(0.0f)
+	, OutputMax(1.0f)
+{
+}
+
+void FSoundModulationInputTransform::Apply(float& Value) const
+{
+	const float Alpha = FMath::Clamp((Value - InputMin) / (InputMax - InputMin), 0.0f, 1.0f);
+	Value = FMath::Lerp(OutputMin, OutputMax, Alpha);
+}
+
+FSoundModulationOutputTransform::FSoundModulationOutputTransform()
+	: InputMin(0.0f)
+	, InputMax(1.0f)
+	, Curve(ESoundModulatorOutputCurve::Exp)
+	, Scalar(2.5f)
+	, OutputMin(0.0f)
+	, OutputMax(1.0f)
+{
+}
+
+void FSoundModulationOutputTransform::Apply(float& Value) const
+{
+	// Clamp the input
+	Value = FMath::Clamp(Value, InputMin, InputMax);
+
+	EvaluateCurve(Value);
+
+	// Clamp the output
+	Value = FMath::Clamp(Value, OutputMin, OutputMax);
+}
+
+void FSoundModulationOutputTransform::EvaluateCurve(float& Value) const
+{
+	// If custom curve, evaluate curve and return before calculating alpha
+	if (Curve == ESoundModulatorOutputCurve::Custom)
+	{
+		if (const FRichCurve* RichCurve = FloatCurve.GetRichCurveConst())
+		{
+			Value = RichCurve->Eval(Value);
+		}
+		return;
+	}
+
+	// Avoid divide-by-zero & return max if input range is limited to single value
+	if (InputMax == InputMin)
+	{
+		Value = OutputMax;
+		return;
+	}
+	
+	const float Alpha = FMath::Clamp((Value - InputMin) / (InputMax - InputMin), 0.0f, 1.0f);
+	switch (Curve)
+	{
+		case ESoundModulatorOutputCurve::Linear:
+		{
+			Value = Alpha;
+		}
+		break;
+	
+		case ESoundModulatorOutputCurve::Exp:
+		{
+			// Alpha is limited to between 0.0f and 1.0f and ExponentialScalar
+			// between 0 and 10 to keep values "sane" and avoid float boundary.
+			Value = Alpha * (FMath::Pow(10.0f, Scalar * (Alpha - 1.0f)));
+		}
+		break;
+
+		case ESoundModulatorOutputCurve::Exp_Inverse:
+		{
+			// Alpha is limited to between 0.0f and 1.0f and ExponentialScalar
+			// between 0 and 10 to keep values "sane" and avoid float boundary.
+			Value = ((Alpha - 1.0f) * (Scalar * FMath::Pow(10.0f, Alpha))) + 1.0f;
+		}
+		break;
+
+		case ESoundModulatorOutputCurve::Log:
+		{
+			Value = (Scalar * FMath::LogX(10.0f, Alpha)) + 1.0f;
+		}
+		break;
+
+		case ESoundModulatorOutputCurve::Sin:
+		{
+			Value = Audio::FastSin(HALF_PI * Alpha);
+		}
+		break;
+
+		case ESoundModulatorOutputCurve::SCurve:
+		{
+			Value = 0.5f * Audio::FastSin(((PI * Alpha) - HALF_PI)) + 0.5f;
+		}
+		break;
+
+		default:
+		{
+			static_assert(static_cast<int32>(ESoundModulatorOutputCurve::Count) == 7, "Possible missing case coverage for output curve.");
+		}
+		break;
+	}
+
+	Value = FMath::Lerp(OutputMin, OutputMax, Value);
+}
diff --git a/Engine/Plugins/Runtime/AudioModulation/Source/AudioModulation/Private/SoundModulationTransform.h b/Engine/Plugins/Runtime/AudioModulation/Source/AudioModulation/Private/SoundModulationTransform.h
new file mode 100644
--- /dev/null
+++ b/Engine/Plugins/Runtime/AudioModulation/Source/AudioModulation/Private/SoundModulationTransform.h
@@ -0,0 +1,96 @@
+// Copyright 1998-2019 Epic Games, Inc. All Rights Reserved.
+#pragma once
+
+#include "CoreMinimal.h"
+#include "Curves/CurveFloat.h"
+#include "UObject/ObjectMacros.h"
+#include "UObject/Object.h"
+
+#include "SoundModulationTransform.generated.h"
+
+
+USTRUCT(BlueprintType)
+struct FSoundModulationInputTransform
+{
+	GENERATED_USTRUCT_BODY()
+
+	FSoundModulationInputTransform();
+
+	/** Minimum value to clamp the input to. */
+	UPROPERTY(EditAnywhere, Category = Modulation, BlueprintReadWrite, meta = (DisplayName = "Input Min", UIMin = "0", UIMax = "1"))
+	float InputMin;
+
+	/** Maximum value to clamp the input to. */
+	UPROPERTY(EditAnywhere, Category = Modulation, BlueprintReadWrite, meta = (DisplayName = "Input Max", UIMin = "0", UIMax = "1"))
+	float InputMax;
+
+	/** Minimum value to clamp to map the output to. */
+	UPROPERTY(EditAnywhere, Category = Modulation, BlueprintReadWrite, meta = (DisplayName = "Output Min", UIMin = "0", UIMax = "1"))
+	float OutputMin;
+
+	/** Maximum value to clamp to map the output to. */
+	UPROPERTY(EditAnywhere, Category = Modulation, BlueprintReadWrite, meta = (DisplayName = "Output Max", UIMin = "0", UIMax = "1"))
+	float OutputMax;
+
+	/** Applies transform to provided value */
+	void Apply(float& OutValue) const;
+};
+
+UENUM()
+enum class ESoundModulatorOutputCurve
+{
+	Linear		UMETA(DisplayName = "Linear"),
+	Exp			UMETA(DisplayName = "Exponential"),
+	Exp_Inverse UMETA(DisplayName = "Exponential (Inverse)"),
+	Log			UMETA(DisplayName = "Log"),
+	Sin			UMETA(DisplayName = "Sin (Quarter)"),
+	SCurve		UMETA(DisplayName = "Sin (S-Curve)"),
+	Custom		UMETA(DisplayName = "Custom"),
+	Count		UMETA(Hidden),
+};
+
+USTRUCT(BlueprintType)
+struct FSoundModulationOutputTransform
+{
+	GENERATED_USTRUCT_BODY()
+
+	FSoundModulationOutputTransform();
+
+	/** Minimum value to clamp the input to. */
+	UPROPERTY(EditAnywhere, Category = Input, BlueprintReadWrite, meta = (DisplayName = "Input Min", UIMin = "0", UIMax = "1"))
+	float InputMin;
+
+	/** Maximum value to clamp the input to. */
+	UPROPERTY(EditAnywhere, Category = Input, BlueprintReadWrite, meta = (DisplayName = "Input Max", UIMin = "0", UIMax = "1"))
+	float InputMax;
+
+	/** The curve to apply when transforming the output. */
+	UPROPERTY(EditAnywhere, Category = Curve, BlueprintReadWrite)
+	ESoundModulatorOutputCurve Curve;
+
+	/** When curve set to log, exponential or exponential inverse, value is factor 'b' in following equations with output 'y' and input 'x':
+	 *  Exponential: y = x * 10^-b(1-x)
+	 *  Exponential (Inverse): y = ((x - 1) * 10^(-bx)) + 1
+	 *  Logarithmic: y = b * log(x) + 1
+	 */
+	UPROPERTY(EditAnywhere, Category = Input, BlueprintReadWrite, meta = (DisplayName = "Exponential Scalar", ClampMin = "0", ClampMax = "10", UIMin = "0", UIMax = "10"))
+	float Scalar;
+
+	/** Custom curve to apply if output curve type is set to custom. */
+	UPROPERTY(EditAnywhere, Category = Curve, BlueprintReadWrite)
+	FRuntimeFloatCurve FloatCurve;
+
+	/** Minimum value to clamp to map the output to. */
+	UPROPERTY(EditAnywhere, Category = Output, BlueprintReadWrite)
+	float OutputMin;
+
+	/** Maximum value to clamp to map the output to. */
+	UPROPERTY(EditAnywhere, Category = Output, BlueprintReadWrite)
+	float OutputMax;
+
+	/** Applies transform to provided value */
+	void Apply(float& OutValue) const;
+
+private:
+	void EvaluateCurve(float& Value) const;
+};
diff --git a/Engine/Plugins/Runtime/AudioModulation/Source/AudioModulation/Private/SoundModulationValue.cpp b/Engine/Plugins/Runtime/AudioModulation/Source/AudioModulation/Private/SoundModulationValue.cpp
new file mode 100644
--- /dev/null
+++ b/Engine/Plugins/Runtime/AudioModulation/Source/AudioModulation/Private/SoundModulationValue.cpp
@@ -0,0 +1,73 @@
+// Copyright 1998-2019 Epic Games, Inc. All Rights Reserved.
+#include "SoundModulationValue.h"
+
+
+namespace AudioModulation
+{
+	const BusMixId InvalidBusMixId = TNumericLimits<uint32>::Max();
+	const BusId    InvalidBusId    = TNumericLimits<uint32>::Max();
+	const BusId    InvalidLFOId    = TNumericLimits<uint32>::Max();
+} // namespace AudioModulation
+
+FSoundModulationValue::FSoundModulationValue()
+	: TargetValue(1.0f)
+	, AttackTime(0.0f)
+	, ReleaseTime(0.0f)
+	, Value(1.0f)
+{
+}
+
+FSoundModulationValue::FSoundModulationValue(float InValue, float InAttackTime, float InReleaseTime)
+	: TargetValue(InValue)
+	, AttackTime(InAttackTime)
+	, ReleaseTime(InReleaseTime)
+	, Value(InValue)
+{
+}
+
+float FSoundModulationValue::GetCurrentValue() const
+{
+	// Current does not require update to be called if
+	// value is attacking or releasing and time is 0
+	if (AttackTime <= 0.0f && Value < TargetValue)
+	{
+		return TargetValue;
+	}
+
+	if (ReleaseTime <= 0.0f && Value > TargetValue)
+	{
+		return TargetValue;
+	}
+
+	return Value;
+}
+
+void FSoundModulationValue::Update(float Elapsed)
+{
+	// Attacking
+	if (Value < TargetValue)
+	{
+		if (AttackTime > 0.0f)
+		{
+			Value = Value + (Elapsed / AttackTime);
+			Value = FMath::Min(Value, TargetValue);
+		}
+		else
+		{
+			Value = TargetValue;
+		}
+	}
+	// Releasing
+	else if (Value > TargetValue)
+	{
+		if (ReleaseTime > 0.0f)
+		{
+			Value = Value - (Elapsed / ReleaseTime);
+			Value = FMath::Max(Value, TargetValue);
+		}
+		else
+		{
+			Value = TargetValue;
+		}
+	}
+}
\ No newline at end of file
diff --git a/Engine/Plugins/Runtime/AudioModulation/Source/AudioModulation/Private/SoundModulationValue.h b/Engine/Plugins/Runtime/AudioModulation/Source/AudioModulation/Private/SoundModulationValue.h
new file mode 100644
--- /dev/null
+++ b/Engine/Plugins/Runtime/AudioModulation/Source/AudioModulation/Private/SoundModulationValue.h
@@ -0,0 +1,49 @@
+// Copyright 1998-2019 Epic Games, Inc. All Rights Reserved.
+#pragma once
+
+#include "UObject/ObjectMacros.h"
+#include "UObject/Object.h"
+
+#include "SoundModulationValue.generated.h"
+
+
+namespace AudioModulation
+{
+	using BusMixId = uint32;
+	extern const BusMixId InvalidBusMixId;
+
+	using BusId = uint32;
+	extern const BusId InvalidBusId;
+
+	using LFOId = uint32;
+	extern const LFOId InvalidLFOId;
+} // namespace AudioModulation
+
+USTRUCT(BlueprintType)
+struct FSoundModulationValue
+{
+	GENERATED_USTRUCT_BODY()
+
+	FSoundModulationValue();
+	FSoundModulationValue(float InValue, float InAttackTime, float InReleaseTime);
+
+	/** Target value of the modulator. */
+	UPROPERTY(EditAnywhere, Category = Modulation, BlueprintReadWrite, meta = (DisplayName = "Value", ClampMin = "0.0", UIMin = "0.0", ClampMax = "1.0", UIMax = "1.0"))
+	float TargetValue;
+
+	/** Time it takes (in sec) to unitarily increase the bus value (from 0 to 1). */
+	UPROPERTY(EditAnywhere, Category = Modulation, BlueprintReadWrite, meta = (DisplayName = "Attack Time", ClampMin = "0.0", UIMin = "0.0"))
+	float AttackTime;
+
+	/** Time it takes (in sec) to unitarily decrease the bus value (from 1 to 0). */
+	UPROPERTY(EditAnywhere, Category = Modulation, BlueprintReadWrite, meta = (DisplayName = "Release Time", ClampMin = "0.0", UIMin = "0.0"))
+	float ReleaseTime;
+
+	/** Current value lerping toward target */
+	float GetCurrentValue() const;
+
+	void Update(float Elapsed);
+
+private:
+	float Value;
+};
diff --git a/Engine/Plugins/Runtime/AudioModulation/Source/AudioModulation/Private/SoundModulatorBase.cpp b/Engine/Plugins/Runtime/AudioModulation/Source/AudioModulation/Private/SoundModulatorBase.cpp
new file mode 100644
--- /dev/null
+++ b/Engine/Plugins/Runtime/AudioModulation/Source/AudioModulation/Private/SoundModulatorBase.cpp
@@ -0,0 +1,8 @@
+// Copyright 1998-2019 Epic Games, Inc. All Rights Reserved.
+#include "SoundModulatorBase.h"
+
+
+USoundModulatorBase::USoundModulatorBase(const FObjectInitializer& ObjectInitializer)
+	: Super(ObjectInitializer)
+{
+}
diff --git a/Engine/Plugins/Runtime/AudioModulation/Source/AudioModulation/Private/SoundModulatorBase.h b/Engine/Plugins/Runtime/AudioModulation/Source/AudioModulation/Private/SoundModulatorBase.h
new file mode 100644
--- /dev/null
+++ b/Engine/Plugins/Runtime/AudioModulation/Source/AudioModulation/Private/SoundModulatorBase.h
@@ -0,0 +1,14 @@
+// Copyright 1998-2019 Epic Games, Inc. All Rights Reserved.
+#pragma once
+
+#include "CoreMinimal.h"
+#include "IAudioExtensionPlugin.h"
+
+#include "SoundModulatorBase.generated.h"
+
+
+UCLASS(hidecategories = Object, abstract, MinimalAPI)
+class USoundModulatorBase : public UObject
+{
+	GENERATED_UCLASS_BODY()
+};
diff --git a/Engine/Plugins/Runtime/AudioModulation/Source/AudioModulation/Private/SoundModulatorBus.cpp b/Engine/Plugins/Runtime/AudioModulation/Source/AudioModulation/Private/SoundModulatorBus.cpp
new file mode 100644
--- /dev/null
+++ b/Engine/Plugins/Runtime/AudioModulation/Source/AudioModulation/Private/SoundModulatorBus.cpp
@@ -0,0 +1,275 @@
+// Copyright 1998-2019 Epic Games, Inc. All Rights Reserved.
+#include "SoundModulatorBus.h"
+
+#include "AudioDevice.h"
+#include "AudioDeviceManager.h"
+#include "AudioModulation.h"
+#include "AudioModulationInternal.h"
+#include "AudioModulationLogging.h"
+#include "Engine/World.h"
+#include "SoundModulatorLFO.h"
+
+
+USoundModulatorBusBase::USoundModulatorBusBase(const FObjectInitializer& ObjectInitializer)
+	: Super(ObjectInitializer)
+	, bAutoActivate(0)
+	, bAutoDeactivate(0)
+	, DefaultValue(1.0f)
+	, Min(0.0f)
+	, Max(1.0f)
+{
+}
+
+#if WITH_EDITOR
+void USoundModulatorBusBase::PostEditChangeProperty(struct FPropertyChangedEvent& PropertyChangedEvent)
+{
+	Super::PostEditChangeProperty(PropertyChangedEvent);
+
+	FAudioDeviceManager* DeviceManager = GEngine->GetAudioDeviceManager();
+	if (!DeviceManager || !PropertyChangedEvent.Property)
+	{
+		return;
+	}
+
+	const TArray<FAudioDevice*>& Devices = DeviceManager->GetAudioDevices();
+	for (FAudioDevice* Device : Devices)
+	{
+		if (!Device || !Device->IsModulationPluginEnabled() || !Device->ModulationInterface.IsValid())
+		{
+			continue;
+		}
+
+		auto Impl = static_cast<AudioModulation::FAudioModulation*>(Device->ModulationInterface.Get())->GetImpl();
+
+		FName Name = PropertyChangedEvent.Property->GetFName();
+		if (Name == TEXT("DefaultValue"))
+		{
+			Impl->SetBusDefault(*this, DefaultValue);
+		}
+		else if (Name == TEXT("Min") || Name == TEXT("Max"))
+		{
+			Impl->SetBusRange(*this, FVector2D(Min, Max));
+		}
+		else if (Name == TEXT("Modulators"))
+		{
+			const AudioModulation::BusId BusId = static_cast<AudioModulation::BusId>(GetUniqueID());
+			if (Impl->IsBusActive(BusId))
+			{
+				Impl->DeactivateBus(BusId);
+				Impl->ActivateBus(*this);
+			}
+		}
+	}
+}
+#endif // WITH_EDITOR
+
+USoundVolumeModulatorBus::USoundVolumeModulatorBus(const FObjectInitializer& ObjectInitializer)
+	: Super(ObjectInitializer)
+{
+}
+
+USoundPitchModulatorBus::USoundPitchModulatorBus(const FObjectInitializer& ObjectInitializer)
+	: Super(ObjectInitializer)
+{
+}
+
+USoundLPFModulatorBus::USoundLPFModulatorBus(const FObjectInitializer& ObjectInitializer)
+	: Super(ObjectInitializer)
+{
+}
+
+USoundHPFModulatorBus::USoundHPFModulatorBus(const FObjectInitializer& ObjectInitializer)
+	: Super(ObjectInitializer)
+{
+}
+
+void USoundModulatorBusBase::BeginDestroy()
+{
+	Super::BeginDestroy();
+
+	UWorld* World = GetWorld();
+	if (!World)
+	{
+		return;
+	}
+
+	FAudioDevice* AudioDevice = World->GetAudioDevice();
+	if (!AudioDevice)
+	{
+		return;
+	}
+
+	check(AudioDevice->IsModulationPluginEnabled());
+	if (IAudioModulation* ModulationInterface = AudioDevice->ModulationInterface.Get())
+	{
+		auto ModulationImpl = static_cast<AudioModulation::FAudioModulation*>(ModulationInterface)->GetImpl();
+		check(ModulationImpl);
+
+		auto BusId = static_cast<const AudioModulation::BusId>(GetUniqueID());
+		ModulationImpl->DeactivateBus(BusId);
+	}
+}
+
+namespace AudioModulation
+{
+	FModulatorBusProxy::FModulatorBusProxy()
+		: BusId(0)
+		, DefaultValue(0.0f)
+		, LFOValue(1.0f)
+		, MixValue(NAN)
+		, Operator(ESoundModulatorOperator::Multiply)
+		, Range(0.0f, 1.0f)
+		, bAutoDeactivate(false) 
+	{
+	}
+
+	FModulatorBusProxy::FModulatorBusProxy(const USoundModulatorBusBase& Bus)
+		: BusId(static_cast<AudioModulation::BusId>(Bus.GetUniqueID()))
+#if !UE_BUILD_SHIPPING
+		, Name(Bus.GetName())
+#endif // !UE_BUILD_SHIPPING
+		, DefaultValue(Bus.DefaultValue)
+		, LFOValue(1.0f)
+		, MixValue(NAN)
+		, Operator(Bus.GetOperator())
+		, Range(Bus.Min, Bus.Max)
+		, bAutoDeactivate(Bus.bAutoDeactivate)
+	{
+		if (Bus.Min > Bus.Max)
+		{
+			Range.X = Bus.Max;
+			Range.Y = Bus.Min;
+		}
+
+		DefaultValue = FMath::Clamp(DefaultValue, Range.X, Range.Y);
+
+		for (const USoundModulatorBase* Modulator: Bus.Modulators)
+		{
+			if (const USoundModulatorLFO* LFO = Cast<USoundModulatorLFO>(Modulator))
+			{
+				LFOId Id = static_cast<LFOId>(LFO->GetUniqueID());
+				LFOIds.Add(Id);
+			}
+		}
+	}
+
+	bool FModulatorBusProxy::CanDeactivate() const
+	{
+		return bAutoDeactivate;
+	}
+
+	AudioModulation::BusId FModulatorBusProxy::GetBusId() const
+	{
+		return BusId;
+	}
+
+	float FModulatorBusProxy::GetDefaultValue() const
+	{
+		return DefaultValue;
+	}
+
+	float FModulatorBusProxy::GetLFOValue() const
+	{
+		return LFOValue;
+	}
+
+	FVector2D FModulatorBusProxy::GetRange() const
+	{
+		return Range;
+	}
+
+	float FModulatorBusProxy::GetMixValue() const
+	{
+		return MixValue;
+	}
+
+#if !UE_BUILD_SHIPPING
+	const FString& FModulatorBusProxy::GetName() const
+	{
+		return Name;
+	}
+#endif // !UE_BUILD_SHIPPING
+
+	float FModulatorBusProxy::GetValue() const
+	{
+		const float DefaultMixed = Mix(DefaultValue);
+		return FMath::Clamp(Mix(DefaultMixed, LFOValue), Range.X, Range.Y);
+	}
+
+	float FModulatorBusProxy::Mix(float ValueA) const
+	{
+		// If mix value is NaN, it is uninitialized (effectively, the parent bus is inactive)
+		// and therefore not mixable, so just return the second value.
+		if (FMath::IsNaN(MixValue))
+		{
+			return ValueA;
+		}
+		return Mix(MixValue, ValueA);
+	}
+		
+	float FModulatorBusProxy::Mix(float ValueA, float ValueB) const
+	{
+		switch (Operator)
+		{
+			case ESoundModulatorOperator::Min:
+			{
+				return FMath::Min(ValueA, ValueB);
+			}
+			break;
+
+			case ESoundModulatorOperator::Max:
+			{
+				return FMath::Max(ValueA, ValueB);
+			}
+			break;
+
+			case ESoundModulatorOperator::Multiply:
+			{
+				return ValueA * ValueB;
+			}
+			break;
+
+			default:
+			{
+				static_assert(static_cast<int32>(ESoundModulatorOperator::Count) == 3, "Missing mod operator case coverage");
+			}
+			break;
+		}
+
+		check(false);
+		return NAN;
+	}
+
+	void FModulatorBusProxy::MixIn(const float InValue)
+	{
+		MixValue = Mix(InValue);
+	}
+
+	void FModulatorBusProxy::MixLFO(LFOProxyMap& LFOMap)
+	{
+		for (const LFOId& LFOId : LFOIds)
+		{
+			if (FModulatorLFOProxy* LFOProxy = LFOMap.Find(LFOId))
+			{
+				LFOProxy->SetIsActive();
+				LFOValue *= LFOProxy->GetValue();
+			}
+		}
+	}
+
+	void FModulatorBusProxy::Reset()
+	{
+		LFOValue = 1.0f;
+		MixValue = NAN;
+	}
+
+	void FModulatorBusProxy::SetDefaultValue(const float Value)
+	{
+		DefaultValue = FMath::Clamp(Value, Range.X, Range.Y);
+	}
+
+	void FModulatorBusProxy::SetRange(const FVector2D& InRange)
+	{
+		Range = InRange;
+	}
+} // namespace AudioModulation
diff --git a/Engine/Plugins/Runtime/AudioModulation/Source/AudioModulation/Private/SoundModulatorBus.h b/Engine/Plugins/Runtime/AudioModulation/Source/AudioModulation/Private/SoundModulatorBus.h
new file mode 100644
--- /dev/null
+++ b/Engine/Plugins/Runtime/AudioModulation/Source/AudioModulation/Private/SoundModulatorBus.h
@@ -0,0 +1,148 @@
+// Copyright 1998-2019 Epic Games, Inc. All Rights Reserved.
+#pragma once
+
+#include "CoreMinimal.h"
+#include "UObject/ObjectMacros.h"
+#include "UObject/Object.h"
+
+#include "SoundModulatorLFO.h"
+#include "SoundModulationValue.h"
+
+#include "SoundModulatorBus.generated.h"
+
+
+UENUM()
+enum class ESoundModulatorOperator : uint8
+{
+	/** Multiply all mix values together */
+	Multiply,
+
+	/** Take the lowest mix value */
+	Min,
+
+	/** Take the highest mix value */
+	Max,
+
+	Count UMETA(Hidden),
+};
+
+
+UCLASS(BlueprintType, hidecategories = Object, abstract, MinimalAPI)
+class USoundModulatorBusBase : public UObject
+{
+	GENERATED_UCLASS_BODY()
+
+public:
+	/** Automatically activates bus when no sounds are being updated by it. */
+	UPROPERTY(EditAnywhere, Category = Modulation, BlueprintReadOnly)
+	uint8 bAutoActivate : 1;
+
+	/** Automatically deactivates bus when no sounds are being updated by it. */
+	UPROPERTY(EditAnywhere, Category = Modulation, BlueprintReadOnly)
+	uint8 bAutoDeactivate : 1;
+
+	/** Default value of modulator (when no mix is applied). */
+	UPROPERTY(EditAnywhere, Category = Modulation, BlueprintReadWrite, meta = (UIMin = "0", UIMax = "1"))
+	float DefaultValue;
+
+	/** Minimum value the bus can achieve (applied post mix phase, pre patch output) */
+	UPROPERTY(EditAnywhere, Category = Modulation, BlueprintReadWrite, meta = (UIMin = "0", UIMax = "1"))
+	float Min;
+
+	/** Maximum value the bus can achieve (applied post mix phase, pre patch output) */
+	UPROPERTY(EditAnywhere, Category = Modulation, BlueprintReadWrite, meta = (UIMin = "0", UIMax = "1"))
+	float Max;
+
+	UPROPERTY(EditAnywhere, Category = Modulation, BlueprintReadWrite)
+	TArray<USoundModulatorBase*> Modulators;
+
+	virtual void BeginDestroy() override;
+
+#if WITH_EDITOR
+	virtual void PostEditChangeProperty(FPropertyChangedEvent& PropertyChangedEvent) override;
+#endif // WITH_EDITOR
+
+	virtual ESoundModulatorOperator GetOperator() const { return ESoundModulatorOperator::Multiply; }
+};
+
+UCLASS(BlueprintType, hidecategories = Object, editinlinenew, MinimalAPI)
+class USoundVolumeModulatorBus : public USoundModulatorBusBase
+{
+	GENERATED_UCLASS_BODY()
+};
+
+UCLASS(BlueprintType, hidecategories = Object, editinlinenew, MinimalAPI)
+class USoundPitchModulatorBus : public USoundModulatorBusBase
+{
+	GENERATED_UCLASS_BODY()
+
+	virtual ESoundModulatorOperator GetOperator() const { return ESoundModulatorOperator::Multiply; }
+};
+
+UCLASS(BlueprintType, hidecategories = Object, editinlinenew, MinimalAPI)
+class USoundHPFModulatorBus : public USoundModulatorBusBase
+{
+	GENERATED_UCLASS_BODY()
+
+	virtual ESoundModulatorOperator GetOperator() const { return ESoundModulatorOperator::Max; }
+};
+
+UCLASS(BlueprintType, hidecategories = Object, editinlinenew, MinimalAPI)
+class USoundLPFModulatorBus : public USoundModulatorBusBase
+{
+	GENERATED_UCLASS_BODY()
+
+	virtual ESoundModulatorOperator GetOperator() const { return ESoundModulatorOperator::Min; }
+};
+
+namespace AudioModulation
+{
+	class FModulatorBusProxy
+	{
+	public:
+		FModulatorBusProxy();
+		FModulatorBusProxy(const USoundModulatorBusBase& Bus);
+
+		bool CanDeactivate() const;
+
+		AudioModulation::BusId GetBusId() const;
+		float GetDefaultValue() const;
+		float GetLFOValue() const;
+		float GetMixValue() const;
+		FVector2D GetRange() const;
+		float GetValue() const;
+		void MixIn(const float InValue);
+		void MixLFO(LFOProxyMap& LFOMap);
+		void Reset();
+
+		void SetDefaultValue(const float Value);
+		void SetRange(const FVector2D& Range);
+
+#if !UE_BUILD_SHIPPING
+		const FString& GetName() const;
+#endif // !UE_BUILD_SHIPPING
+
+	private:
+		float Mix(float ValueA) const;
+		float Mix(float ValueA, float ValueB) const;
+
+		BusId BusId;
+
+#if !UE_BUILD_SHIPPING
+		FString Name;
+#endif // !UE_BUILD_SHIPPING
+
+		float DefaultValue;
+
+		// Cached values
+		float LFOValue;
+		float MixValue;
+
+		TArray<LFOId> LFOIds;
+		ESoundModulatorOperator Operator;
+		FVector2D Range;
+
+		uint8 bAutoDeactivate : 1;
+	};
+	using BusProxyMap = TMap<BusId, FModulatorBusProxy>;
+} // namespace AudioModulation
diff --git a/Engine/Plugins/Runtime/AudioModulation/Source/AudioModulation/Private/SoundModulatorBusMix.cpp b/Engine/Plugins/Runtime/AudioModulation/Source/AudioModulation/Private/SoundModulatorBusMix.cpp
new file mode 100644
--- /dev/null
+++ b/Engine/Plugins/Runtime/AudioModulation/Source/AudioModulation/Private/SoundModulatorBusMix.cpp
@@ -0,0 +1,159 @@
+// Copyright 1998-2019 Epic Games, Inc. All Rights Reserved.
+#include "SoundModulatorBusMix.h"
+
+#include "AudioDevice.h"
+#include "AudioModulation.h"
+#include "AudioModulationInternal.h"
+#include "AudioModulationLogging.h"
+#include "Engine/World.h"
+
+
+FSoundModulatorBusMixChannel::FSoundModulatorBusMixChannel()
+	: Bus(nullptr)
+{
+}
+
+FSoundModulatorBusMixChannel::FSoundModulatorBusMixChannel(USoundModulatorBusBase* InBus, const float TargetValue)
+	: Bus(InBus)
+{
+	Value.TargetValue = FMath::Clamp(TargetValue, 0.0f, 1.0f);
+}
+
+USoundModulatorBusMix::USoundModulatorBusMix(const FObjectInitializer& ObjectInitializer)
+	: Super(ObjectInitializer)
+{
+}
+
+void USoundModulatorBusMix::BeginDestroy()
+{
+	Super::BeginDestroy();
+
+	UWorld* World = GetWorld();
+	if (!World)
+	{
+		return;
+	}
+
+	FAudioDevice* AudioDevice = World->GetAudioDevice();
+	if (!AudioDevice)
+	{
+		return;
+	}
+
+	check(AudioDevice->IsModulationPluginEnabled());
+	if (IAudioModulation* ModulationInterface = AudioDevice->ModulationInterface.Get())
+	{
+		auto ModulationImpl = static_cast<AudioModulation::FAudioModulation*>(ModulationInterface)->GetImpl();
+		check(ModulationImpl);
+
+		auto BusMixId = static_cast<const AudioModulation::BusMixId>(GetUniqueID());
+		ModulationImpl->DeactivateBusMix(BusMixId);
+	}
+}
+
+namespace AudioModulation
+{
+	FModulatorBusMixChannelProxy::FModulatorBusMixChannelProxy(const FSoundModulatorBusMixChannel& Channel)
+		: Value(Channel.Value)
+	{
+		check(Channel.Bus);
+		BusId = static_cast<BusMixId>(Channel.Bus->GetUniqueID());
+
+#if !UE_BUILD_SHIPPING
+		Name = Channel.Bus->GetName();
+#endif // !UE_BUILD_SHIPPING
+	}
+
+	FModulatorBusMixProxy::FModulatorBusMixProxy(const USoundModulatorBusMix& Mix)
+#if UE_BUILD_SHIPPING
+		: Status(BusMixStatus::Active)
+#else
+		: Name(Mix.GetName())
+		, Status(BusMixStatus::Active)
+#endif // !UE_BUILD_SHIPPING
+	{
+		for (const FSoundModulatorBusMixChannel& Channel : Mix.Channels)
+		{
+			if (Channel.Bus)
+			{
+				auto BusMixId = static_cast<const AudioModulation::BusMixId>(Channel.Bus->GetUniqueID());
+#if !UE_BUILD_SHIPPING
+				if (Channels.Contains(BusMixId))
+				{
+					UE_LOG(LogAudioModulation, Warning,
+						TEXT("USoundModulatorBusMix '%s' already contains bus '%s'. Only one representative channel for this bus added."),
+						*Mix.GetFullName(), *Channel.Bus->GetFullName());
+				}
+#endif // !UE_BUILD_SHIPPING
+				Channels.Emplace(BusMixId, FModulatorBusMixChannelProxy(Channel));
+			}
+			else
+			{
+				UE_LOG(LogAudioModulation, Warning,
+					TEXT("USoundModulatorBusMix '%s' has channel with no bus specified. "
+					"Mix activated but channel ignored."),
+					*Mix.GetFullName());
+			}
+		}
+	}
+
+	bool FModulatorBusMixProxy::CanDeactivate() const
+	{
+		return Status == BusMixStatus::Stopped;
+	}
+
+	void FModulatorBusMixProxy::SetActive()
+	{
+		Status = BusMixStatus::Active;
+	}
+
+	void FModulatorBusMixProxy::SetStopping()
+	{
+		if (Status == BusMixStatus::Active)
+		{
+			Status = BusMixStatus::Stopping;
+		}
+	}
+
+	void FModulatorBusMixProxy::Update(float Elapsed, BusProxyMap& ProxyMap)
+	{
+		bool bRequestStop = true;
+		for (TPair<BusId, FModulatorBusMixChannelProxy>& Channel : Channels)
+		{
+			FModulatorBusMixChannelProxy& ChannelProxy = Channel.Value;
+			FSoundModulationValue& MixChannelValue = ChannelProxy.Value;
+
+			if (FModulatorBusProxy* BusProxy = ProxyMap.Find(ChannelProxy.BusId))
+			{
+				MixChannelValue.Update(Elapsed);
+
+				const float CurrentValue = MixChannelValue.GetCurrentValue();
+				if (Status == BusMixStatus::Stopping)
+				{
+					MixChannelValue.TargetValue = BusProxy->GetDefaultValue();
+					if (!FMath::IsNearlyEqual(MixChannelValue.TargetValue, CurrentValue))
+					{
+						bRequestStop = false;
+					}
+				}
+				else
+				{
+					bRequestStop = false;
+				}
+				BusProxy->MixIn(CurrentValue);
+			}
+		}
+
+		if (bRequestStop)
+		{
+			Status = BusMixStatus::Stopped;
+		}
+	}
+
+#if !UE_BUILD_SHIPPING
+	const FString& FModulatorBusMixProxy::GetName() const
+	{
+		return Name;
+	}
+#endif // !UE_BUILD_SHIPPING
+} // namespace AudioModulation
diff --git a/Engine/Plugins/Runtime/AudioModulation/Source/AudioModulation/Private/SoundModulatorBusMix.h b/Engine/Plugins/Runtime/AudioModulation/Source/AudioModulation/Private/SoundModulatorBusMix.h
new file mode 100644
--- /dev/null
+++ b/Engine/Plugins/Runtime/AudioModulation/Source/AudioModulation/Private/SoundModulatorBusMix.h
@@ -0,0 +1,89 @@
+// Copyright 1998-2019 Epic Games, Inc. All Rights Reserved.
+#pragma once
+
+#include "CoreMinimal.h"
+#include "UObject/ObjectMacros.h"
+#include "UObject/Object.h"
+
+#include "SoundModulatorBus.h"
+#include "SoundModulationValue.h"
+
+#include "SoundModulatorBusMix.generated.h"
+
+
+USTRUCT(BlueprintType)
+struct FSoundModulatorBusMixChannel
+{
+	GENERATED_USTRUCT_BODY()
+
+	FSoundModulatorBusMixChannel();
+	FSoundModulatorBusMixChannel(USoundModulatorBusBase* InBus, const float TargetValue);
+
+	/* Bus controlled by channel. */
+	UPROPERTY(EditAnywhere, Category = Channel, BlueprintReadOnly)
+	USoundModulatorBusBase* Bus;
+
+	/* Value mix is set to. */
+	UPROPERTY(EditAnywhere, Category = Channel, BlueprintReadWrite, meta = (ShowOnlyInnerProperties))
+	FSoundModulationValue Value;
+};
+
+UCLASS(config = Engine, autoexpandcategories = (Channel, Mix), editinlinenew, BlueprintType, MinimalAPI)
+class USoundModulatorBusMix : public UObject
+{
+	GENERATED_UCLASS_BODY()
+
+public:
+
+	virtual void BeginDestroy() override;
+
+	/* Array of channels controlled by mix. */
+	UPROPERTY(EditAnywhere, Category = Mix, BlueprintReadOnly)
+	TArray<FSoundModulatorBusMixChannel> Channels;
+};
+
+namespace AudioModulation
+{
+	struct FModulatorBusMixChannelProxy
+	{
+		FModulatorBusMixChannelProxy(const FSoundModulatorBusMixChannel& Channel);
+
+		BusId BusId;
+		FSoundModulationValue Value;
+		
+#if !UE_BUILD_SHIPPING
+		FString Name;
+#endif // !UE_BUILD_SHIPPING
+	};
+
+	struct FModulatorBusMixProxy
+	{
+		enum class BusMixStatus : uint8
+		{
+			Active,
+			Stopping,
+			Stopped
+		};
+
+		FModulatorBusMixProxy(const USoundModulatorBusMix& Mix);
+
+		bool CanDeactivate() const;
+		void SetActive();
+		void SetStopping();
+		void Update(const float Elapsed, BusProxyMap& ProxyMap);
+
+#if !UE_BUILD_SHIPPING
+		const FString& GetName() const;
+#endif // !UE_BUILD_SHIPPING
+
+		TMap<BusId, FModulatorBusMixChannelProxy> Channels;
+
+	private:
+#if !UE_BUILD_SHIPPING
+		FString Name;
+#endif // !UE_BUILD_SHIPPING
+
+		BusMixStatus Status;
+	};
+	using BusMixProxyMap = TMap<BusMixId, FModulatorBusMixProxy>;
+} // namespace AudioModulation
diff --git a/Engine/Plugins/Runtime/AudioModulation/Source/AudioModulation/Private/SoundModulatorLFO.cpp b/Engine/Plugins/Runtime/AudioModulation/Source/AudioModulation/Private/SoundModulatorLFO.cpp
new file mode 100644
--- /dev/null
+++ b/Engine/Plugins/Runtime/AudioModulation/Source/AudioModulation/Private/SoundModulatorLFO.cpp
@@ -0,0 +1,143 @@
+// Copyright 1998-2019 Epic Games, Inc. All Rights Reserved.
+#include "SoundModulatorLFO.h"
+
+#include "AudioDevice.h"
+#include "AudioModulation.h"
+#include "AudioModulationInternal.h"
+#include "AudioModulationLogging.h"
+#include "Engine/World.h"
+
+
+USoundModulatorLFO::USoundModulatorLFO(const FObjectInitializer& ObjectInitializer)
+	: Super(ObjectInitializer)
+	, Shape(ESoundModulatorLFOShape::Sine)
+	, Amplitude(0.5f)
+	, Frequency(1.0f)
+	, Offset(0.5f)
+	, bLooping(1)
+	, bAutoActivate(0)
+	, bAutoDeactivate(0)
+{
+}
+
+void USoundModulatorLFO::BeginDestroy()
+{
+	Super::BeginDestroy();
+
+	UWorld* World = GetWorld();
+	if (!World)
+	{
+		return;
+	}
+
+	FAudioDevice* AudioDevice = World->GetAudioDevice();
+	if (!AudioDevice)
+	{
+		return;
+	}
+
+	check(AudioDevice->IsModulationPluginEnabled());
+	if (IAudioModulation* ModulationInterface = AudioDevice->ModulationInterface.Get())
+	{
+		auto ModulationImpl = static_cast<AudioModulation::FAudioModulation*>(ModulationInterface)->GetImpl();
+		check(ModulationImpl);
+
+		auto LFOId = static_cast<const AudioModulation::LFOId>(GetUniqueID());
+		ModulationImpl->DeactivateLFO(LFOId);
+	}
+}
+
+namespace AudioModulation
+{
+	FModulatorLFOProxy::FModulatorLFOProxy()
+		: Id(0)
+		, Offset(0.0f)
+		, Value(0.0f)
+		, bIsActive(false)
+		, bAutoDeactivate(false)
+	{
+	}
+
+	FModulatorLFOProxy::FModulatorLFOProxy(const USoundModulatorLFO& InLFO)
+		: Id(static_cast<LFOId>(InLFO.GetUniqueID()))
+#if !UE_BUILD_SHIPPING
+		, Name(InLFO.GetName())
+#endif // !UE_BUILD_SHIPPING
+		, Offset(InLFO.Offset)
+		, Value(0.0f)
+		, bIsActive(false)
+		, bAutoDeactivate(InLFO.bAutoDeactivate)
+	{
+		LFO.SetGain(InLFO.Amplitude);
+		LFO.SetFrequency(InLFO.Frequency);
+		LFO.SetMode(InLFO.bLooping ? Audio::ELFOMode::Type::Sync : Audio::ELFOMode::OneShot);
+
+		static_assert(static_cast<int32>(ESoundModulatorLFOShape::COUNT) == static_cast<int32>(Audio::ELFO::Type::NumLFOTypes), "LFOShape/ELFO Type mismatch");
+		LFO.SetType(static_cast<Audio::ELFO::Type>(InLFO.Shape));
+
+		LFO.Start();
+	}
+
+	bool FModulatorLFOProxy::CanDeactivate() const
+	{
+		return !bIsActive && bAutoDeactivate;
+	}
+
+	void FModulatorLFOProxy::ClearIsActive()
+	{
+		bIsActive = false;
+	}
+
+	float FModulatorLFOProxy::GetAmplitude() const
+	{
+		return LFO.GetGain();
+	}
+
+	float FModulatorLFOProxy::GetFreq() const
+	{
+		return LFO.GetFrequency();
+	}
+
+#if !UE_BUILD_SHIPPING
+	const FString& FModulatorLFOProxy::GetName() const
+	{
+		return Name;
+	}
+#endif // !UE_BUILD_SHIPPING
+
+	LFOId FModulatorLFOProxy::GetId() const
+	{
+		return Id;
+	}
+
+	float FModulatorLFOProxy::GetOffset() const
+	{
+		return Offset;
+	}
+
+	float FModulatorLFOProxy::GetValue() const
+	{
+		return Value;
+	}
+
+	void FModulatorLFOProxy::SetFreq(float InFreq)
+	{
+		LFO.SetFrequency(InFreq);
+	}
+
+	void FModulatorLFOProxy::SetIsActive()
+	{
+		bIsActive = true;
+	}
+
+	void FModulatorLFOProxy::Update(float InElapsed)
+	{
+		if (InElapsed > 0.0f && LFO.GetFrequency() > 0.0f)
+		{
+			const float SampleRate = 1.0f / InElapsed;
+			LFO.SetSampleRate(SampleRate);
+			LFO.Update();
+			Value = LFO.Generate() + Offset;
+		}
+	}
+} // namespace AudioModulation
diff --git a/Engine/Plugins/Runtime/AudioModulation/Source/AudioModulation/Private/SoundModulatorLFO.h b/Engine/Plugins/Runtime/AudioModulation/Source/AudioModulation/Private/SoundModulatorLFO.h
new file mode 100644
--- /dev/null
+++ b/Engine/Plugins/Runtime/AudioModulation/Source/AudioModulation/Private/SoundModulatorLFO.h
@@ -0,0 +1,103 @@
+// Copyright 1998-2019 Epic Games, Inc. All Rights Reserved.
+#pragma once
+
+#include "CoreMinimal.h"
+#include "DSP/LFO.h"
+#include "UObject/ObjectMacros.h"
+#include "UObject/Object.h"
+
+#include "SoundModulatorBase.h"
+#include "SoundModulationValue.h"
+
+#include "SoundModulatorLFO.generated.h"
+
+UENUM(BlueprintType)
+enum class ESoundModulatorLFOShape : uint8
+{
+	Sine			  UMETA(DisplayName = "Sine"),
+	UpSaw			  UMETA(DisplayName = "Saw (Up)"),
+	DownSaw			  UMETA(DisplayName = "Saw (Down)"),
+	Square			  UMETA(DisplayName = "Square"),
+	Triangle		  UMETA(DisplayName = "Triangle"),
+	Exponential		  UMETA(DisplayName = "Exponential"),
+	RandomSampleHold  UMETA(DisplayName = "Random"),
+
+	COUNT UMETA(Hidden)
+};
+
+UCLASS(BlueprintType, hidecategories = Object, editinlinenew, MinimalAPI)
+class USoundModulatorLFO : public USoundModulatorBase
+{
+	GENERATED_UCLASS_BODY()
+
+public:
+	virtual void BeginDestroy() override;
+
+	/** Shape of oscillating waveform */
+	UPROPERTY(EditAnywhere, Category = Modulation, BlueprintReadWrite)
+	ESoundModulatorLFOShape Shape;
+
+	/** Amplitude of oscillator */
+	UPROPERTY(EditAnywhere, Category = Modulation, BlueprintReadWrite, meta = (UIMin = "0", UIMax = "1", ClampMin = "0", ClampMax = "1"))
+	float Amplitude;
+
+	/** Frequency of oscillator */
+	UPROPERTY(EditAnywhere, Category = Modulation, BlueprintReadWrite, meta = (UIMin = "0", UIMax = "20", ClampMin = "0", ClampMax = "20"))
+	float Frequency;
+
+	/** Offset of oscillator */
+	UPROPERTY(EditAnywhere, Category = Modulation, BlueprintReadWrite, meta = (UIMin = "0", UIMax = "1", ClampMin = "0", ClampMax = "1"))
+	float Offset;
+
+	/** Whether or not to loop the oscillation more than once */
+	UPROPERTY(EditAnywhere, Category = Modulation, BlueprintReadWrite)
+	uint8 bLooping : 1;
+
+	/** Automatically activates LFO when sounds are referencing it */
+	UPROPERTY(EditAnywhere, Category = Modulation, BlueprintReadWrite)
+	uint8 bAutoActivate : 1;
+
+	/** Automatically deactivates LFO when no sounds are referencing it */
+	UPROPERTY(EditAnywhere, Category = Modulation, BlueprintReadWrite)
+	uint8 bAutoDeactivate : 1;
+};
+
+namespace AudioModulation
+{
+	class FModulatorLFOProxy
+	{
+	public:
+		FModulatorLFOProxy();
+		FModulatorLFOProxy(const USoundModulatorLFO& InLFO);
+
+		bool CanDeactivate() const;
+		void ClearIsActive();
+		float GetAmplitude() const;
+		float GetFreq() const;
+		LFOId GetId() const;
+#if !UE_BUILD_SHIPPING
+		const FString& GetName() const;
+#endif // !UE_BUILD_SHIPPING
+		float GetOffset() const;
+		float GetValue() const;
+		void SetFreq(float InFreq);
+		void SetIsActive();
+		void Update(float InElapsed);
+
+	private:
+		LFOId Id;
+
+#if !UE_BUILD_SHIPPING
+		FString Name;
+#endif // !UE_BUILD_SHIPPING
+
+		Audio::FLFO LFO;
+		float Offset;
+		float Value;
+
+		uint8 bLimitedLifetime : 1;
+		uint8 bIsActive : 1;
+		uint8 bAutoDeactivate : 1;
+	};
+	using LFOProxyMap = TMap<BusId, FModulatorLFOProxy>;
+} // namespace AudioModulation
diff --git a/Engine/Plugins/Runtime/AudioModulation/Source/AudioModulation/Public/AudioModulation.h b/Engine/Plugins/Runtime/AudioModulation/Source/AudioModulation/Public/AudioModulation.h
new file mode 100644
--- /dev/null
+++ b/Engine/Plugins/Runtime/AudioModulation/Source/AudioModulation/Public/AudioModulation.h
@@ -0,0 +1,81 @@
+// Copyright 1998-2019 Epic Games, Inc. All Rights Reserved.
+#pragma once
+
+#include "CoreMinimal.h"
+#include "HAL/ThreadSafeBool.h"
+#include "IAudioExtensionPlugin.h"
+#include "Modules/ModuleInterface.h"
+#include "SoundModulationPatch.h"
+#include "Stats/Stats.h"
+
+
+// Cycle stats for audio mixer
+DECLARE_STATS_GROUP(TEXT("AudioModulation"), STATGROUP_AudioModulation, STATCAT_Advanced);
+
+// Tracks the time for the full render block 
+DECLARE_CYCLE_STAT_EXTERN(TEXT("Process Modulators"), STAT_AudioModulationProcessModulators, STATGROUP_AudioModulation, AUDIOMODULATION_API);
+DECLARE_CYCLE_STAT_EXTERN(TEXT("Process Controls"), STAT_AudioModulationProcessControls, STATGROUP_AudioModulation, AUDIOMODULATION_API);
+
+namespace AudioModulation
+{
+	class FAudioModulationImpl;
+
+	class AUDIOMODULATION_API FAudioModulation : public IAudioModulation
+	{
+	public:
+		FAudioModulation();
+		virtual ~FAudioModulation() = default;
+
+		//~ Begin IAudioModulation implementation
+		virtual void Initialize(const FAudioPluginInitializationParams& InitializationParams) override;
+
+#if WITH_EDITOR
+		virtual void OnEditSource(const USoundModulationPluginSourceSettingsBase& Settings) override;
+#endif // WITH_EDITOR
+
+		virtual void OnInitSource(const uint32 SourceId, const FName& AudioComponentUserId, const uint32 NumChannels, USoundModulationPluginSourceSettingsBase* Settings) override;
+
+#if !UE_BUILD_SHIPPING
+		virtual bool OnPostHelp(FCommonViewportClient* ViewportClient, const TCHAR* Stream) override;
+		virtual int32 OnRenderStat(FViewport* Viewport, FCanvas* Canvas, int32 X, int32 Y, const UFont& Font, const FVector* ViewLocation, const FRotator* ViewRotation) override;
+		virtual bool OnToggleStat(FCommonViewportClient* ViewportClient, const TCHAR* Stream) override;
+#endif // !UE_BUILD_SHIPPING
+
+		virtual void OnReleaseSource(const uint32 SourceId) override;
+		virtual void ProcessControls(const uint32 SourceId, FSoundModulationControls& Controls) override;
+		virtual void ProcessModulators(const float Elapsed) override;
+		//~ End IAudioModulation implementation
+
+		FAudioModulationImpl* GetImpl();
+
+	private:
+		TUniquePtr<FAudioModulationImpl> Impl;
+	};
+} // namespace AudioModulation
+
+class FAudioModulationPluginFactory : public IAudioModulationFactory
+{
+public:
+	virtual const FName& GetDisplayName() const override
+	{
+		static FName DisplayName = FName(TEXT("DefaultModulationPlugin"));
+		return DisplayName;
+	}
+
+	virtual TAudioModulationPtr CreateNewModulationPlugin(FAudioDevice* OwningDevice) override;
+
+	virtual UClass* GetCustomModulationSettingsClass() const override
+	{
+		return USoundModulationSettings::StaticClass();
+	}
+};
+
+class FAudioModulationModule : public IModuleInterface
+{
+public:
+	virtual void StartupModule() override;
+	virtual void ShutdownModule() override;
+
+private:
+	FAudioModulationPluginFactory ModulationPluginFactory;
+};
\ No newline at end of file
diff --git a/Engine/Plugins/Runtime/AudioModulation/Source/AudioModulationEditor/AudioModulationEditor.build.cs b/Engine/Plugins/Runtime/AudioModulation/Source/AudioModulationEditor/AudioModulationEditor.build.cs
new file mode 100644
--- /dev/null
+++ b/Engine/Plugins/Runtime/AudioModulation/Source/AudioModulationEditor/AudioModulationEditor.build.cs
@@ -0,0 +1,33 @@
+// Copyright 1998-2019 Epic Games, Inc. All Rights Reserved.
+
+using UnrealBuildTool;
+
+public class AudioModulationEditor : ModuleRules
+{
+	public AudioModulationEditor(ReadOnlyTargetRules Target) : base(Target)
+	{
+		OptimizeCode = CodeOptimization.Never;
+
+		PrivateDependencyModuleNames.AddRange(
+			new string[] {
+				"Core",
+				"CoreUObject",
+				"Engine",
+				"UnrealEd",
+				"PropertyEditor",
+				"SequenceRecorder",
+				"Slate",
+				"SlateCore",
+				"AudioEditor",
+				"AudioModulation",
+			}
+		);
+
+		PrivateIncludePaths.AddRange(
+			new string[] {
+				"AudioModulationEditor/Private",
+				"AudioModulation/Private"
+			}
+		);
+	}
+}
diff --git a/Engine/Plugins/Runtime/AudioModulation/Source/AudioModulationEditor/Private/AssetTypeActions/AssetTypeActions_SoundModulationSettings.cpp b/Engine/Plugins/Runtime/AudioModulation/Source/AudioModulationEditor/Private/AssetTypeActions/AssetTypeActions_SoundModulationSettings.cpp
new file mode 100644
--- /dev/null
+++ b/Engine/Plugins/Runtime/AudioModulation/Source/AudioModulationEditor/Private/AssetTypeActions/AssetTypeActions_SoundModulationSettings.cpp
@@ -0,0 +1,10 @@
+// Copyright 1998-2019 Epic Games, Inc. All Rights Reserved.
+#include "AssetTypeActions_SoundModulationSettings.h"
+
+#include "SoundModulationPatch.h"
+
+
+UClass* FAssetTypeActions_SoundModulationSettings::GetSupportedClass() const
+{
+	return USoundModulationSettings::StaticClass();
+}
diff --git a/Engine/Plugins/Runtime/AudioModulation/Source/AudioModulationEditor/Private/AssetTypeActions/AssetTypeActions_SoundModulationSettings.h b/Engine/Plugins/Runtime/AudioModulation/Source/AudioModulationEditor/Private/AssetTypeActions/AssetTypeActions_SoundModulationSettings.h
new file mode 100644
--- /dev/null
+++ b/Engine/Plugins/Runtime/AudioModulation/Source/AudioModulationEditor/Private/AssetTypeActions/AssetTypeActions_SoundModulationSettings.h
@@ -0,0 +1,16 @@
+// Copyright 1998-2019 Epic Games, Inc. All Rights Reserved.
+
+#pragma once
+
+#include "CoreMinimal.h"
+#include "AssetTypeActions_Base.h"
+
+class FAssetTypeActions_SoundModulationSettings : public FAssetTypeActions_Base
+{
+public:
+	// IAssetTypeActions Implementation
+	virtual FText GetName() const override { return NSLOCTEXT("AssetTypeActions", "AssetTypeActions_SoundModulationSettings", "Modulation Settings"); }
+	virtual FColor GetTypeColor() const override { return FColor(237, 237, 0); }
+	virtual UClass* GetSupportedClass() const override;
+	virtual uint32 GetCategories() override { return EAssetTypeCategories::Sounds; }
+};
diff --git a/Engine/Plugins/Runtime/AudioModulation/Source/AudioModulationEditor/Private/AssetTypeActions/AssetTypeActions_SoundModulatorBus.cpp b/Engine/Plugins/Runtime/AudioModulation/Source/AudioModulationEditor/Private/AssetTypeActions/AssetTypeActions_SoundModulatorBus.cpp
new file mode 100644
--- /dev/null
+++ b/Engine/Plugins/Runtime/AudioModulation/Source/AudioModulationEditor/Private/AssetTypeActions/AssetTypeActions_SoundModulatorBus.cpp
@@ -0,0 +1,25 @@
+// Copyright 1998-2019 Epic Games, Inc. All Rights Reserved.
+
+#include "AssetTypeActions_SoundModulatorBus.h"
+#include "SoundModulatorBus.h"
+
+
+UClass* FAssetTypeActions_SoundVolumeModulatorBus::GetSupportedClass() const
+{
+	return USoundVolumeModulatorBus::StaticClass();
+}
+
+UClass* FAssetTypeActions_SoundPitchModulatorBus::GetSupportedClass() const
+{
+	return USoundPitchModulatorBus::StaticClass();
+}
+
+UClass* FAssetTypeActions_SoundLPFModulatorBus::GetSupportedClass() const
+{
+	return USoundLPFModulatorBus::StaticClass();
+}
+
+UClass* FAssetTypeActions_SoundHPFModulatorBus::GetSupportedClass() const
+{
+	return USoundHPFModulatorBus::StaticClass();
+}
diff --git a/Engine/Plugins/Runtime/AudioModulation/Source/AudioModulationEditor/Private/AssetTypeActions/AssetTypeActions_SoundModulatorBus.h b/Engine/Plugins/Runtime/AudioModulation/Source/AudioModulationEditor/Private/AssetTypeActions/AssetTypeActions_SoundModulatorBus.h
new file mode 100644
--- /dev/null
+++ b/Engine/Plugins/Runtime/AudioModulation/Source/AudioModulationEditor/Private/AssetTypeActions/AssetTypeActions_SoundModulatorBus.h
@@ -0,0 +1,46 @@
+// Copyright 1998-2019 Epic Games, Inc. All Rights Reserved.
+
+#pragma once
+
+#include "CoreMinimal.h"
+#include "AssetTypeActions_Base.h"
+
+class FAssetTypeActions_SoundVolumeModulatorBus : public FAssetTypeActions_Base
+{
+public:
+	// IAssetTypeActions Implementation
+	virtual FText GetName() const override { return NSLOCTEXT("AssetTypeActions", "AssetTypeActions_SoundVolumeModulatorBus", "Modulation Bus (Volume)"); }
+	virtual FColor GetTypeColor() const override { return FColor(33, 183, 0); }
+	virtual UClass* GetSupportedClass() const override;
+	virtual uint32 GetCategories() override { return EAssetTypeCategories::Sounds; }
+};
+
+class FAssetTypeActions_SoundPitchModulatorBus : public FAssetTypeActions_Base
+{
+public:
+	// IAssetTypeActions Implementation
+	virtual FText GetName() const override { return NSLOCTEXT("AssetTypeActions", "AssetTypeActions_SoundPitchModulatorBus", "Modulation Bus (Pitch)"); }
+	virtual FColor GetTypeColor() const override { return FColor(181, 21, 0); }
+	virtual UClass* GetSupportedClass() const override;
+	virtual uint32 GetCategories() override { return EAssetTypeCategories::Sounds; }
+};
+
+class FAssetTypeActions_SoundLPFModulatorBus : public FAssetTypeActions_Base
+{
+public:
+	// IAssetTypeActions Implementation
+	virtual FText GetName() const override { return NSLOCTEXT("AssetTypeActions", "AssetTypeActions_SoundLPFModulatorBus", "Modulation Bus (LPF)"); }
+	virtual FColor GetTypeColor() const override { return FColor(0, 156, 183); }
+	virtual UClass* GetSupportedClass() const override;
+	virtual uint32 GetCategories() override { return EAssetTypeCategories::Sounds; }
+};
+
+class FAssetTypeActions_SoundHPFModulatorBus : public FAssetTypeActions_Base
+{
+public:
+	// IAssetTypeActions Implementation
+	virtual FText GetName() const override { return NSLOCTEXT("AssetTypeActions", "AssetTypeActions_SoundHPFModulatorBus", "Modulation Bus (HPF)"); }
+	virtual FColor GetTypeColor() const override { return FColor(94, 237, 183); }
+	virtual UClass* GetSupportedClass() const override;
+	virtual uint32 GetCategories() override { return EAssetTypeCategories::Sounds; }
+};
diff --git a/Engine/Plugins/Runtime/AudioModulation/Source/AudioModulationEditor/Private/AssetTypeActions/AssetTypeActions_SoundModulatorBusMix.cpp b/Engine/Plugins/Runtime/AudioModulation/Source/AudioModulationEditor/Private/AssetTypeActions/AssetTypeActions_SoundModulatorBusMix.cpp
new file mode 100644
--- /dev/null
+++ b/Engine/Plugins/Runtime/AudioModulation/Source/AudioModulationEditor/Private/AssetTypeActions/AssetTypeActions_SoundModulatorBusMix.cpp
@@ -0,0 +1,10 @@
+// Copyright 1998-2019 Epic Games, Inc. All Rights Reserved.
+
+#include "AssetTypeActions_SoundModulatorBusMix.h"
+#include "SoundModulatorBusMix.h"
+
+
+UClass* FAssetTypeActions_SoundModulatorBusMix::GetSupportedClass() const
+{
+	return USoundModulatorBusMix::StaticClass();
+}
diff --git a/Engine/Plugins/Runtime/AudioModulation/Source/AudioModulationEditor/Private/AssetTypeActions/AssetTypeActions_SoundModulatorBusMix.h b/Engine/Plugins/Runtime/AudioModulation/Source/AudioModulationEditor/Private/AssetTypeActions/AssetTypeActions_SoundModulatorBusMix.h
new file mode 100644
--- /dev/null
+++ b/Engine/Plugins/Runtime/AudioModulation/Source/AudioModulationEditor/Private/AssetTypeActions/AssetTypeActions_SoundModulatorBusMix.h
@@ -0,0 +1,16 @@
+// Copyright 1998-2019 Epic Games, Inc. All Rights Reserved.
+
+#pragma once
+
+#include "CoreMinimal.h"
+#include "AssetTypeActions_Base.h"
+
+class FAssetTypeActions_SoundModulatorBusMix : public FAssetTypeActions_Base
+{
+public:
+	// IAssetTypeActions Implementation
+	virtual FText GetName() const override { return NSLOCTEXT("AssetTypeActions", "AssetTypeActions_SoundModulatorBusMix", "Modulation Mix"); }
+	virtual FColor GetTypeColor() const override { return FColor(237, 237, 0); }
+	virtual UClass* GetSupportedClass() const override;
+	virtual uint32 GetCategories() override { return EAssetTypeCategories::Sounds; }
+};
diff --git a/Engine/Plugins/Runtime/AudioModulation/Source/AudioModulationEditor/Private/AssetTypeActions/AssetTypeActions_SoundModulatorLFO.cpp b/Engine/Plugins/Runtime/AudioModulation/Source/AudioModulationEditor/Private/AssetTypeActions/AssetTypeActions_SoundModulatorLFO.cpp
new file mode 100644
--- /dev/null
+++ b/Engine/Plugins/Runtime/AudioModulation/Source/AudioModulationEditor/Private/AssetTypeActions/AssetTypeActions_SoundModulatorLFO.cpp
@@ -0,0 +1,10 @@
+// Copyright 1998-2019 Epic Games, Inc. All Rights Reserved.
+
+#include "AssetTypeActions_SoundModulatorLFO.h"
+#include "SoundModulatorLFO.h"
+
+
+UClass* FAssetTypeActions_SoundModulatorLFO::GetSupportedClass() const
+{
+	return USoundModulatorLFO::StaticClass();
+}
diff --git a/Engine/Plugins/Runtime/AudioModulation/Source/AudioModulationEditor/Private/AssetTypeActions/AssetTypeActions_SoundModulatorLFO.h b/Engine/Plugins/Runtime/AudioModulation/Source/AudioModulationEditor/Private/AssetTypeActions/AssetTypeActions_SoundModulatorLFO.h
new file mode 100644
--- /dev/null
+++ b/Engine/Plugins/Runtime/AudioModulation/Source/AudioModulationEditor/Private/AssetTypeActions/AssetTypeActions_SoundModulatorLFO.h
@@ -0,0 +1,16 @@
+// Copyright 1998-2019 Epic Games, Inc. All Rights Reserved.
+
+#pragma once
+
+#include "CoreMinimal.h"
+#include "AssetTypeActions_Base.h"
+
+class FAssetTypeActions_SoundModulatorLFO : public FAssetTypeActions_Base
+{
+public:
+	// IAssetTypeActions Implementation
+	virtual FText GetName() const override { return NSLOCTEXT("AssetTypeActions", "AssetTypeActions_SoundModulatorLFO", "Modulation LFO"); }
+	virtual FColor GetTypeColor() const override { return FColor(237, 237, 0); }
+	virtual UClass* GetSupportedClass() const override;
+	virtual uint32 GetCategories() override { return EAssetTypeCategories::Sounds; }
+};
diff --git a/Engine/Plugins/Runtime/AudioModulation/Source/AudioModulationEditor/Private/AudioModulationEditor.cpp b/Engine/Plugins/Runtime/AudioModulation/Source/AudioModulationEditor/Private/AudioModulationEditor.cpp
new file mode 100644
--- /dev/null
+++ b/Engine/Plugins/Runtime/AudioModulation/Source/AudioModulationEditor/Private/AudioModulationEditor.cpp
@@ -0,0 +1,69 @@
+// Copyright 1998-2019 Epic Games, Inc. All Rights Reserved.
+#include "AudioModulationEditor.h"
+
+#include "AssetTypeActions/AssetTypeActions_SoundModulationSettings.h"
+#include "AssetTypeActions/AssetTypeActions_SoundModulatorBus.h"
+#include "AssetTypeActions/AssetTypeActions_SoundModulatorBusMix.h"
+#include "AssetTypeActions/AssetTypeActions_SoundModulatorLFO.h"
+#include "Layouts/SoundModulationTransformLayout.h"
+#include "SoundModulationTransform.h"
+#include "Styling/SlateStyle.h"
+#include "Styling/SlateStyleRegistry.h"
+
+namespace
+{
+	TSharedPtr<FSlateStyleSet> StyleSet;
+} // namespace <>
+
+FAudioModulationEditorModule::FAudioModulationEditorModule()
+{
+	StyleSet = MakeShareable(new FSlateStyleSet("AudioModulationStyleSet"));
+}
+
+void FAudioModulationEditorModule::SetIcon(const FString& ClassName)
+{
+	static const FVector2D Icon16 = FVector2D(16.0f, 16.0f);
+	static const FVector2D Icon64 = FVector2D(64.0f, 64.0f);
+
+	static const FString IconDir = FPaths::EngineDir() / FString::Printf(TEXT("Plugins/Runtime/AudioModulation/Icons"));
+
+	const FString IconFileName16 = FString::Printf(TEXT("%s_16x.png"), *ClassName);
+	const FString IconFileName64 = FString::Printf(TEXT("%s_64x.png"), *ClassName);
+
+	StyleSet->Set(*FString::Printf(TEXT("ClassIcon.%s"), *ClassName), new FSlateImageBrush(IconDir / IconFileName16, Icon16));
+	StyleSet->Set(*FString::Printf(TEXT("ClassThumbnail.%s"), *ClassName), new FSlateImageBrush(IconDir / IconFileName64, Icon64));
+}
+
+void FAudioModulationEditorModule::StartupModule()
+{
+	// Register the audio editor asset type actions
+	IAssetTools& AssetTools = FModuleManager::LoadModuleChecked<FAssetToolsModule>("AssetTools").Get();
+
+	AssetTools.RegisterAssetTypeActions(MakeShareable(new FAssetTypeActions_SoundVolumeModulatorBus));
+	AssetTools.RegisterAssetTypeActions(MakeShareable(new FAssetTypeActions_SoundPitchModulatorBus));
+	AssetTools.RegisterAssetTypeActions(MakeShareable(new FAssetTypeActions_SoundLPFModulatorBus));
+	AssetTools.RegisterAssetTypeActions(MakeShareable(new FAssetTypeActions_SoundHPFModulatorBus));
+	AssetTools.RegisterAssetTypeActions(MakeShareable(new FAssetTypeActions_SoundModulatorBusMix));
+	AssetTools.RegisterAssetTypeActions(MakeShareable(new FAssetTypeActions_SoundModulatorLFO));
+	AssetTools.RegisterAssetTypeActions(MakeShareable(new FAssetTypeActions_SoundModulationSettings));
+
+	SetIcon(TEXT("SoundVolumeModulatorBus"));
+	SetIcon(TEXT("SoundPitchModulatorBus"));
+	SetIcon(TEXT("SoundLPFModulatorBus"));
+	SetIcon(TEXT("SoundHPFModulatorBus"));
+	SetIcon(TEXT("SoundModulatorBusMix"));
+	SetIcon(TEXT("SoundModulatorLFO"));
+	SetIcon(TEXT("SoundModulationSettings"));
+
+	FPropertyEditorModule& PropertyModule = FModuleManager::LoadModuleChecked<FPropertyEditorModule>("PropertyEditor");
+	PropertyModule.RegisterCustomPropertyTypeLayout("SoundModulationOutputTransform", FOnGetPropertyTypeCustomizationInstance::CreateStatic(&FSoundModulationOutputTransformLayoutCustomization::MakeInstance));
+
+	FSlateStyleRegistry::RegisterSlateStyle(*StyleSet.Get());
+}
+
+void FAudioModulationEditorModule::ShutdownModule()
+{
+}
+
+
+IMPLEMENT_MODULE(FAudioModulationEditorModule, AudioModulationEditor);
diff --git a/Engine/Plugins/Runtime/AudioModulation/Source/AudioModulationEditor/Private/AudioModulationEditor.h b/Engine/Plugins/Runtime/AudioModulation/Source/AudioModulationEditor/Private/AudioModulationEditor.h
new file mode 100644
--- /dev/null
+++ b/Engine/Plugins/Runtime/AudioModulation/Source/AudioModulationEditor/Private/AudioModulationEditor.h
@@ -0,0 +1,18 @@
+// Copyright 1998-2019 Epic Games, Inc. All Rights Reserved.
+#pragma once
+
+#include "CoreMinimal.h"
+#include "Modules/ModuleManager.h"
+
+
+class FAudioModulationEditorModule : public IModuleInterface
+{
+public:
+	FAudioModulationEditorModule();
+
+	virtual void StartupModule() override;
+	virtual void ShutdownModule() override;
+
+private:
+	void SetIcon(const FString& ClassName);
+};
diff --git a/Engine/Plugins/Runtime/AudioModulation/Source/AudioModulationEditor/Private/Factories/SoundModulationSettingsFactory.cpp b/Engine/Plugins/Runtime/AudioModulation/Source/AudioModulationEditor/Private/Factories/SoundModulationSettingsFactory.cpp
new file mode 100644
--- /dev/null
+++ b/Engine/Plugins/Runtime/AudioModulation/Source/AudioModulationEditor/Private/Factories/SoundModulationSettingsFactory.cpp
@@ -0,0 +1,18 @@
+// Copyright 1998-2019 Epic Games, Inc. All Rights Reserved.
+#include "SoundModulationSettingsFactory.h"
+#include "SoundModulationPatch.h"
+
+
+USoundModulationSettingsFactory::USoundModulationSettingsFactory(const FObjectInitializer& ObjectInitializer)
+	: Super(ObjectInitializer)
+{
+	SupportedClass = USoundModulationSettings::StaticClass();
+	bCreateNew = true;
+	bEditorImport = false;
+	bEditAfterNew = true;
+}
+
+UObject* USoundModulationSettingsFactory::FactoryCreateNew(UClass* Class, UObject* InParent, FName Name, EObjectFlags Flags, UObject* Context, FFeedbackContext* Warn)
+{
+	return NewObject<USoundModulationSettings>(InParent, Name, Flags);
+}
diff --git a/Engine/Plugins/Runtime/AudioModulation/Source/AudioModulationEditor/Private/Factories/SoundModulationSettingsFactory.h b/Engine/Plugins/Runtime/AudioModulation/Source/AudioModulationEditor/Private/Factories/SoundModulationSettingsFactory.h
new file mode 100644
--- /dev/null
+++ b/Engine/Plugins/Runtime/AudioModulation/Source/AudioModulationEditor/Private/Factories/SoundModulationSettingsFactory.h
@@ -0,0 +1,18 @@
+// Copyright 1998-2019 Epic Games, Inc. All Rights Reserved.
+#pragma once
+
+#include "CoreMinimal.h"
+#include "UObject/ObjectMacros.h"
+#include "Factories/Factory.h"
+#include "SoundModulationSettingsFactory.generated.h"
+
+
+UCLASS(hidecategories=Object, MinimalAPI)
+class USoundModulationSettingsFactory : public UFactory
+{
+	GENERATED_UCLASS_BODY()
+
+	// Begin UFactory Interface
+	virtual UObject* FactoryCreateNew(UClass* Class,UObject* InParent,FName Name,EObjectFlags Flags,UObject* Context,FFeedbackContext* Warn) override;
+	// Begin UFactory Interface	
+};
diff --git a/Engine/Plugins/Runtime/AudioModulation/Source/AudioModulationEditor/Private/Factories/SoundModulatorBusFactory.cpp b/Engine/Plugins/Runtime/AudioModulation/Source/AudioModulationEditor/Private/Factories/SoundModulatorBusFactory.cpp
new file mode 100644
--- /dev/null
+++ b/Engine/Plugins/Runtime/AudioModulation/Source/AudioModulationEditor/Private/Factories/SoundModulatorBusFactory.cpp
@@ -0,0 +1,61 @@
+// Copyright 1998-2019 Epic Games, Inc. All Rights Reserved.
+
+#include "SoundModulatorBusFactory.h"
+#include "SoundModulatorBus.h"
+
+
+USoundVolumeModulatorBusFactory::USoundVolumeModulatorBusFactory(const FObjectInitializer& ObjectInitializer)
+	: Super(ObjectInitializer)
+{
+	SupportedClass = USoundVolumeModulatorBus::StaticClass();
+	bCreateNew     = true;
+	bEditorImport  = false;
+	bEditAfterNew  = true;
+}
+
+UObject* USoundVolumeModulatorBusFactory::FactoryCreateNew(UClass* Class, UObject* InParent, FName Name, EObjectFlags Flags, UObject* Context, FFeedbackContext* Warn)
+{
+	return NewObject<USoundVolumeModulatorBus>(InParent, Name, Flags);
+}
+
+USoundPitchModulatorBusFactory::USoundPitchModulatorBusFactory(const FObjectInitializer& ObjectInitializer)
+	: Super(ObjectInitializer)
+{
+	SupportedClass = USoundPitchModulatorBus::StaticClass();
+	bCreateNew = true;
+	bEditorImport = false;
+	bEditAfterNew = true;
+}
+
+UObject* USoundPitchModulatorBusFactory::FactoryCreateNew(UClass* Class, UObject* InParent, FName Name, EObjectFlags Flags, UObject* Context, FFeedbackContext* Warn)
+{
+	return NewObject<USoundPitchModulatorBus>(InParent, Name, Flags);
+}
+
+USoundLPFModulatorBusFactory::USoundLPFModulatorBusFactory(const FObjectInitializer& ObjectInitializer)
+	: Super(ObjectInitializer)
+{
+	SupportedClass = USoundLPFModulatorBus::StaticClass();
+	bCreateNew = true;
+	bEditorImport = false;
+	bEditAfterNew = true;
+}
+
+UObject* USoundLPFModulatorBusFactory::FactoryCreateNew(UClass* Class, UObject* InParent, FName Name, EObjectFlags Flags, UObject* Context, FFeedbackContext* Warn)
+{
+	return NewObject<USoundLPFModulatorBus>(InParent, Name, Flags);
+}
+
+USoundHPFModulatorBusFactory::USoundHPFModulatorBusFactory(const FObjectInitializer& ObjectInitializer)
+	: Super(ObjectInitializer)
+{
+	SupportedClass = USoundHPFModulatorBus::StaticClass();
+	bCreateNew = true;
+	bEditorImport = false;
+	bEditAfterNew = true;
+}
+
+UObject* USoundHPFModulatorBusFactory::FactoryCreateNew(UClass* Class, UObject* InParent, FName Name, EObjectFlags Flags, UObject* Context, FFeedbackContext* Warn)
+{
+	return NewObject<USoundHPFModulatorBus>(InParent, Name, Flags);
+}
\ No newline at end of file
diff --git a/Engine/Plugins/Runtime/AudioModulation/Source/AudioModulationEditor/Private/Factories/SoundModulatorBusFactory.h b/Engine/Plugins/Runtime/AudioModulation/Source/AudioModulationEditor/Private/Factories/SoundModulatorBusFactory.h
new file mode 100644
--- /dev/null
+++ b/Engine/Plugins/Runtime/AudioModulation/Source/AudioModulationEditor/Private/Factories/SoundModulatorBusFactory.h
@@ -0,0 +1,44 @@
+// Copyright 1998-2019 Epic Games, Inc. All Rights Reserved.
+
+//=============================================================================
+// USoundModulationBusFactory
+//=============================================================================
+
+#pragma once
+
+#include "CoreMinimal.h"
+#include "UObject/ObjectMacros.h"
+#include "Factories/Factory.h"
+#include "SoundModulatorBusFactory.generated.h"
+
+UCLASS(hidecategories=Object, MinimalAPI)
+class USoundVolumeModulatorBusFactory : public UFactory
+{
+	GENERATED_UCLASS_BODY()
+
+	virtual UObject* FactoryCreateNew(UClass* Class, UObject* InParent, FName Name, EObjectFlags Flags, UObject* Context, FFeedbackContext* Warn) override;
+};
+
+UCLASS(hidecategories = Object, MinimalAPI)
+class USoundPitchModulatorBusFactory : public UFactory
+{
+	GENERATED_UCLASS_BODY()
+
+	virtual UObject* FactoryCreateNew(UClass* Class, UObject* InParent, FName Name, EObjectFlags Flags, UObject* Context, FFeedbackContext* Warn) override;
+};
+
+UCLASS(hidecategories = Object, MinimalAPI)
+class USoundHPFModulatorBusFactory : public UFactory
+{
+	GENERATED_UCLASS_BODY()
+
+	virtual UObject* FactoryCreateNew(UClass* Class, UObject* InParent, FName Name, EObjectFlags Flags, UObject* Context, FFeedbackContext* Warn) override;
+};
+
+UCLASS(hidecategories = Object, MinimalAPI)
+class USoundLPFModulatorBusFactory : public UFactory
+{
+	GENERATED_UCLASS_BODY()
+
+	virtual UObject* FactoryCreateNew(UClass* Class, UObject* InParent, FName Name, EObjectFlags Flags, UObject* Context, FFeedbackContext* Warn) override;
+};
\ No newline at end of file
diff --git a/Engine/Plugins/Runtime/AudioModulation/Source/AudioModulationEditor/Private/Factories/SoundModulatorBusMixFactory.cpp b/Engine/Plugins/Runtime/AudioModulation/Source/AudioModulationEditor/Private/Factories/SoundModulatorBusMixFactory.cpp
new file mode 100644
--- /dev/null
+++ b/Engine/Plugins/Runtime/AudioModulation/Source/AudioModulationEditor/Private/Factories/SoundModulatorBusMixFactory.cpp
@@ -0,0 +1,19 @@
+// Copyright 1998-2019 Epic Games, Inc. All Rights Reserved.
+
+#include "SoundModulatorBusMixFactory.h"
+#include "SoundModulatorBusMix.h"
+
+
+USoundModulatorBusMixFactory::USoundModulatorBusMixFactory(const FObjectInitializer& ObjectInitializer)
+	: Super(ObjectInitializer)
+{
+	SupportedClass = USoundModulatorBusMix::StaticClass();
+	bCreateNew     = true;
+	bEditorImport  = false;
+	bEditAfterNew  = true;
+}
+
+UObject* USoundModulatorBusMixFactory::FactoryCreateNew(UClass* Class, UObject* InParent, FName Name, EObjectFlags Flags, UObject* Context, FFeedbackContext* Warn)
+{
+	return NewObject<USoundModulatorBusMix>(InParent, Name, Flags);
+}
diff --git a/Engine/Plugins/Runtime/AudioModulation/Source/AudioModulationEditor/Private/Factories/SoundModulatorBusMixFactory.h b/Engine/Plugins/Runtime/AudioModulation/Source/AudioModulationEditor/Private/Factories/SoundModulatorBusMixFactory.h
new file mode 100644
--- /dev/null
+++ b/Engine/Plugins/Runtime/AudioModulation/Source/AudioModulationEditor/Private/Factories/SoundModulatorBusMixFactory.h
@@ -0,0 +1,20 @@
+// Copyright 1998-2019 Epic Games, Inc. All Rights Reserved.
+
+//=============================================================================
+// USoundModulationBusFactory
+//=============================================================================
+
+#pragma once
+
+#include "CoreMinimal.h"
+#include "UObject/ObjectMacros.h"
+#include "Factories/Factory.h"
+#include "SoundModulatorBusMixFactory.generated.h"
+
+UCLASS(hidecategories=Object, MinimalAPI)
+class USoundModulatorBusMixFactory : public UFactory
+{
+	GENERATED_UCLASS_BODY()
+
+	virtual UObject* FactoryCreateNew(UClass* Class, UObject* InParent, FName Name, EObjectFlags Flags, UObject* Context, FFeedbackContext* Warn) override;
+};
diff --git a/Engine/Plugins/Runtime/AudioModulation/Source/AudioModulationEditor/Private/Factories/SoundModulatorLFOFactory.cpp b/Engine/Plugins/Runtime/AudioModulation/Source/AudioModulationEditor/Private/Factories/SoundModulatorLFOFactory.cpp
new file mode 100644
--- /dev/null
+++ b/Engine/Plugins/Runtime/AudioModulation/Source/AudioModulationEditor/Private/Factories/SoundModulatorLFOFactory.cpp
@@ -0,0 +1,19 @@
+// Copyright 1998-2019 Epic Games, Inc. All Rights Reserved.
+
+#include "SoundModulatorLFOFactory.h"
+#include "SoundModulatorLFO.h"
+
+
+USoundModulatorLFOFactory::USoundModulatorLFOFactory(const FObjectInitializer& ObjectInitializer)
+	: Super(ObjectInitializer)
+{
+	SupportedClass = USoundModulatorLFO::StaticClass();
+	bCreateNew     = true;
+	bEditorImport  = false;
+	bEditAfterNew  = true;
+}
+
+UObject* USoundModulatorLFOFactory::FactoryCreateNew(UClass* Class, UObject* InParent, FName Name, EObjectFlags Flags, UObject* Context, FFeedbackContext* Warn)
+{
+	return NewObject<USoundModulatorLFO>(InParent, Name, Flags);
+}
diff --git a/Engine/Plugins/Runtime/AudioModulation/Source/AudioModulationEditor/Private/Factories/SoundModulatorLFOFactory.h b/Engine/Plugins/Runtime/AudioModulation/Source/AudioModulationEditor/Private/Factories/SoundModulatorLFOFactory.h
new file mode 100644
--- /dev/null
+++ b/Engine/Plugins/Runtime/AudioModulation/Source/AudioModulationEditor/Private/Factories/SoundModulatorLFOFactory.h
@@ -0,0 +1,20 @@
+// Copyright 1998-2019 Epic Games, Inc. All Rights Reserved.
+
+//=============================================================================
+// USoundModulatorLFOFactory
+//=============================================================================
+
+#pragma once
+
+#include "CoreMinimal.h"
+#include "UObject/ObjectMacros.h"
+#include "Factories/Factory.h"
+#include "SoundModulatorLFOFactory.generated.h"
+
+UCLASS(hidecategories=Object, MinimalAPI)
+class USoundModulatorLFOFactory : public UFactory
+{
+	GENERATED_UCLASS_BODY()
+
+	virtual UObject* FactoryCreateNew(UClass* Class, UObject* InParent, FName Name, EObjectFlags Flags, UObject* Context, FFeedbackContext* Warn) override;
+};
diff --git a/Engine/Plugins/Runtime/AudioModulation/Source/AudioModulationEditor/Private/Layouts/SoundModulationTransformLayout.cpp b/Engine/Plugins/Runtime/AudioModulation/Source/AudioModulationEditor/Private/Layouts/SoundModulationTransformLayout.cpp
new file mode 100644
--- /dev/null
+++ b/Engine/Plugins/Runtime/AudioModulation/Source/AudioModulationEditor/Private/Layouts/SoundModulationTransformLayout.cpp
@@ -0,0 +1,85 @@
+// Copyright 1998-2019 Epic Games, Inc. All Rights Reserved.
+#include "SoundModulationTransformLayout.h"
+
+#include "Audio.h"
+#include "DetailCategoryBuilder.h"
+#include "DetailLayoutBuilder.h"
+#include "DetailWidgetRow.h"
+#include "IDetailChildrenBuilder.h"
+#include "IDetailPropertyRow.h"
+#include "PropertyRestriction.h"
+#include "SoundModulationTransform.h"
+#include "Widgets/Text/STextBlock.h"
+
+
+FSoundModulationOutputTransformLayoutCustomization::FSoundModulationOutputTransformLayoutCustomization()
+{
+}
+
+bool FSoundModulationOutputTransformLayoutCustomization::IsCurveSet(TSharedPtr<IPropertyHandle> CurveHandle, const TArray<FString>& Filters) const
+{
+	if (!CurveHandle.IsValid())
+	{
+		return false;
+	}
+
+	FString CurveString;
+	CurveHandle->GetValueAsFormattedString(CurveString);
+	for (const FString& Filter : Filters)
+	{
+		if (CurveString.Contains(*Filter))
+		{
+			return true;
+		}
+	}
+
+	return false;
+}
+
+void FSoundModulationOutputTransformLayoutCustomization::CustomizeHeader(TSharedRef<IPropertyHandle> PropertyHandle, FDetailWidgetRow& HeaderRow, IPropertyTypeCustomizationUtils& CustomizationUtils)
+{
+	// Property handle here is the base struct. We are going to hide it since we're showing it's properties directly.
+// 	PropertyHandle->MarkHiddenByCustomization();
+}
+
+void FSoundModulationOutputTransformLayoutCustomization::CustomizeChildren(TSharedRef<IPropertyHandle> StructPropertyHandle, IDetailChildrenBuilder& ChildBuilder, IPropertyTypeCustomizationUtils& StructCustomizationUtils)
+{
+	uint32 NumChildren;
+	StructPropertyHandle->GetNumChildren(NumChildren);
+
+	TMap<FName, TSharedPtr<IPropertyHandle>> PropertyHandles;
+
+	for (uint32 ChildIndex = 0; ChildIndex < NumChildren; ++ChildIndex)
+	{
+		TSharedRef<IPropertyHandle> ChildHandle = StructPropertyHandle->GetChildHandle(ChildIndex).ToSharedRef();
+		const FName PropertyName = ChildHandle->GetProperty()->GetFName();
+		PropertyHandles.Add(PropertyName, ChildHandle);
+	}
+
+	TSharedRef<IPropertyHandle>InputMinHandle = PropertyHandles.FindChecked(GET_MEMBER_NAME_CHECKED(FSoundModulationOutputTransform, InputMin)).ToSharedRef();
+	ChildBuilder.AddProperty(InputMinHandle);
+
+	TSharedRef<IPropertyHandle>InputMaxHandle = PropertyHandles.FindChecked(GET_MEMBER_NAME_CHECKED(FSoundModulationOutputTransform, InputMax)).ToSharedRef();
+	ChildBuilder.AddProperty(InputMaxHandle);
+
+	TSharedRef<IPropertyHandle>CurveHandle = PropertyHandles.FindChecked(GET_MEMBER_NAME_CHECKED(FSoundModulationOutputTransform, Curve)).ToSharedRef();
+	ChildBuilder.AddProperty(CurveHandle);
+
+	TSharedRef<IPropertyHandle> FloatCurveHandle = PropertyHandles.FindChecked(GET_MEMBER_NAME_CHECKED(FSoundModulationOutputTransform, FloatCurve)).ToSharedRef();
+	const TArray<FString> CustomFilters = TArray<FString>({ TEXT("Custom") });
+	ChildBuilder.AddProperty(FloatCurveHandle)
+		.EditCondition(TAttribute<bool>::Create([this, CurveHandle, CustomFilters]() { return IsCurveSet(CurveHandle, CustomFilters); }), nullptr)
+		.Visibility(TAttribute<EVisibility>::Create([this, CurveHandle, CustomFilters]() { return IsCurveSet(CurveHandle, CustomFilters) ? EVisibility::Visible : EVisibility::Hidden; }));
+
+	TSharedRef<IPropertyHandle> ExpScalarHandle = PropertyHandles.FindChecked(GET_MEMBER_NAME_CHECKED(FSoundModulationOutputTransform, Scalar)).ToSharedRef();
+	const TArray<FString> ScalarFilters = TArray<FString>({ TEXT("Exp"), TEXT("Log") });
+	ChildBuilder.AddProperty(ExpScalarHandle)
+		.EditCondition(TAttribute<bool>::Create([this, CurveHandle, ScalarFilters]() { return IsCurveSet(CurveHandle, ScalarFilters); }), nullptr)
+		.Visibility(TAttribute<EVisibility>::Create([this, CurveHandle, ScalarFilters]() { return IsCurveSet(CurveHandle, ScalarFilters) ? EVisibility::Visible : EVisibility::Hidden; }));
+
+	TSharedRef<IPropertyHandle>OutputMinHandle = PropertyHandles.FindChecked(GET_MEMBER_NAME_CHECKED(FSoundModulationOutputTransform, OutputMin)).ToSharedRef();
+	ChildBuilder.AddProperty(OutputMinHandle);
+
+	TSharedRef<IPropertyHandle>OutputMaxHandle = PropertyHandles.FindChecked(GET_MEMBER_NAME_CHECKED(FSoundModulationOutputTransform, OutputMax)).ToSharedRef();
+	ChildBuilder.AddProperty(OutputMaxHandle);
+}
diff --git a/Engine/Plugins/Runtime/AudioModulation/Source/AudioModulationEditor/Private/Layouts/SoundModulationTransformLayout.h b/Engine/Plugins/Runtime/AudioModulation/Source/AudioModulationEditor/Private/Layouts/SoundModulationTransformLayout.h
new file mode 100644
--- /dev/null
+++ b/Engine/Plugins/Runtime/AudioModulation/Source/AudioModulationEditor/Private/Layouts/SoundModulationTransformLayout.h
@@ -0,0 +1,28 @@
+// Copyright 1998-2019 Epic Games, Inc. All Rights Reserved.
+#pragma once
+
+#include "CoreMinimal.h"
+#include "IPropertyTypeCustomization.h"
+#include "Layout/Visibility.h"
+#include "Misc/Attribute.h"
+#include "PropertyHandle.h"
+
+
+class FSoundModulationOutputTransformLayoutCustomization : public IPropertyTypeCustomization
+{
+public:
+	FSoundModulationOutputTransformLayoutCustomization();
+
+	static TSharedRef<IPropertyTypeCustomization> MakeInstance()
+	{
+		return MakeShareable(new FSoundModulationOutputTransformLayoutCustomization);
+	}
+
+	//~ Begin IPropertyTypeCustomization
+	virtual void CustomizeHeader(TSharedRef<IPropertyHandle> StructPropertyHandle, FDetailWidgetRow& HeaderRow, IPropertyTypeCustomizationUtils& StructCustomizationUtils) override;
+	virtual void CustomizeChildren(TSharedRef<IPropertyHandle> StructPropertyHandle, IDetailChildrenBuilder& ChildBuilder, IPropertyTypeCustomizationUtils& StructCustomizationUtils) override;
+	//~ End IPropertyTypeCustomization
+
+private:
+	bool IsCurveSet(TSharedPtr<IPropertyHandle> CurveHandle, const TArray<FString>& Filter) const;
+};
\ No newline at end of file
diff --git a/Engine/Source/Runtime/AudioMixer/Private/AudioMixerSourceManager.cpp b/Engine/Source/Runtime/AudioMixer/Private/AudioMixerSourceManager.cpp
--- a/Engine/Source/Runtime/AudioMixer/Private/AudioMixerSourceManager.cpp
+++ b/Engine/Source/Runtime/AudioMixer/Private/AudioMixerSourceManager.cpp
@@ -219,7 +219,6 @@ namespace Audio
 			SourceInfo.bUseHRTFSpatializer = false;
 			SourceInfo.bUseOcclusionPlugin = false;
 			SourceInfo.bUseReverbPlugin = false;
-			SourceInfo.bUseModulationPlugin = false;
 			SourceInfo.bHasStarted = false;
 			SourceInfo.bOutputToBusOnly = false;
 			SourceInfo.bIsVorbis = false;
@@ -449,11 +448,6 @@ namespace Audio
 			MixerDevice->ReverbPluginInterface->OnReleaseSource(SourceId);
 		}
 
-		if (SourceInfo.bUseModulationPlugin)
-		{
-			MixerDevice->ModulationInterface->OnReleaseSource(SourceId);
-		}
-
 		// Delete the source effects
 		SourceInfo.SourceEffectChainId = INDEX_NONE;
 		ResetSourceEffectChain(SourceId);
@@ -504,7 +498,6 @@ namespace Audio
 		SourceInfo.bIsExternalSend = false;
 		SourceInfo.bUseOcclusionPlugin = false;
 		SourceInfo.bUseReverbPlugin = false;
-		SourceInfo.bUseModulationPlugin = false;
 		SourceInfo.bHasStarted = false;
 		SourceInfo.bOutputToBusOnly = false;
 		SourceInfo.bIsBypassingLPF = false;
@@ -617,6 +610,12 @@ namespace Audio
 		// Make sure we flag that this source needs a speaker map to at least get one
 		GameThreadInfo.bNeedsSpeakerMap[SourceId] = true;
 
+		// Create the modulation plugin source effect
+		if (InitParams.ModulationPluginSettings != nullptr)
+		{
+			MixerDevice->ModulationInterface->OnInitSource(SourceId, InitParams.AudioComponentUserID, InitParams.NumInputChannels, InitParams.ModulationPluginSettings);
+		}
+
 		AudioMixerThreadCommand([this, SourceId, InitParams]()
 		{
 			AUDIO_MIXER_CHECK_AUDIO_PLAT_THREAD(MixerDevice);
@@ -675,13 +674,6 @@ namespace Audio
 				SourceInfo.bUseReverbPlugin = true;
 			}
 
-			// Create the modulation plugin source effect
-			if (InitParams.ModulationPluginSettings != nullptr)
-			{
-				MixerDevice->ModulationInterface->OnInitSource(SourceId, InitParams.AudioComponentUserID, InitParams.NumInputChannels, InitParams.ModulationPluginSettings);
-				SourceInfo.bUseModulationPlugin = true;
-			}
-
 			// Default all sounds to not consider effect chain tails when playing
 			SourceInfo.bEffectTailsDone = true;
 
@@ -834,6 +826,11 @@ namespace Audio
 
 		AUDIO_MIXER_CHECK(GameThreadInfo.FreeSourceIndices.Contains(SourceId));
 
+		if (MixerDevice->ModulationInterface)
+		{
+			MixerDevice->ModulationInterface->OnReleaseSource(SourceId);
+		}
+
 		AudioMixerThreadCommand([this, SourceId]()
 		{
 			AUDIO_MIXER_CHECK_AUDIO_PLAT_THREAD(MixerDevice);
diff --git a/Engine/Source/Runtime/AudioMixer/Private/AudioMixerSourceManager.h b/Engine/Source/Runtime/AudioMixer/Private/AudioMixerSourceManager.h
--- a/Engine/Source/Runtime/AudioMixer/Private/AudioMixerSourceManager.h
+++ b/Engine/Source/Runtime/AudioMixer/Private/AudioMixerSourceManager.h
@@ -497,7 +497,6 @@ namespace Audio
 			uint8 bIsExternalSend:1;
 			uint8 bUseOcclusionPlugin:1;
 			uint8 bUseReverbPlugin:1;
-			uint8 bUseModulationPlugin:1;
 			uint8 bIsDone:1;
 			uint8 bIsLastBuffer:1;
 			uint8 bOutputToBusOnly:1;
diff --git a/Engine/Source/Runtime/Engine/Private/AudioDevice.cpp b/Engine/Source/Runtime/Engine/Private/AudioDevice.cpp
--- a/Engine/Source/Runtime/Engine/Private/AudioDevice.cpp
+++ b/Engine/Source/Runtime/Engine/Private/AudioDevice.cpp
@@ -445,7 +445,7 @@ bool FAudioDevice::Init(int32 InMaxSources)
 		ModulationInterface->Initialize(PluginInitializationParams);
 
 		bModulationInterfaceEnabled = true;
-		UE_LOG(LogAudio, Display, TEXT("Audio Modulation Plugin: %s"), *(ModulationPluginFactory->GetDisplayName()));
+		UE_LOG(LogAudio, Display, TEXT("Audio Modulation Plugin: %s"), *(ModulationPluginFactory->GetDisplayName().ToString()));
 	}
 
 	// allow the platform to startup
diff --git a/Engine/Source/Runtime/Engine/Private/AudioPluginUtilities.cpp b/Engine/Source/Runtime/Engine/Private/AudioPluginUtilities.cpp
--- a/Engine/Source/Runtime/Engine/Private/AudioPluginUtilities.cpp
+++ b/Engine/Source/Runtime/Engine/Private/AudioPluginUtilities.cpp
@@ -1,9 +1,20 @@
 // Copyright 1998-2019 Epic Games, Inc. All Rights Reserved.
 
 #include "AudioPluginUtilities.h"
+
+#include "CoreGlobals.h"
+#include "CoreMinimal.h"
 #include "Features/IModularFeatures.h"
 #include "Misc/ConfigCacheIni.h"
-#include "CoreGlobals.h"
+
+
+static FString DefaultModulationPluginCVar = TEXT("");
+FAutoConsoleVariableRef CVarActiveModulationPlugin(
+	TEXT("au.DefaultModulationPlugin"),
+	DefaultModulationPluginCVar,
+	TEXT("Name of default modulation plugin to load and use (overridden "
+	"by platform-specific implementation name in config.\n"),
+	ECVF_Default);
 
 
 /************************************************************************/
@@ -179,18 +190,17 @@ IAudioOcclusionFactory* AudioPluginUtilities::GetDesiredOcclusionPlugin(EAudioPl
 
 IAudioModulationFactory* AudioPluginUtilities::GetDesiredModulationPlugin(EAudioPlatform AudioPlatform)
 {
-	//Get the name of the desired spatialization plugin:
-	FString DesiredModulationPlugin = GetDesiredPluginName(EAudioPlugin::MODULATION, AudioPlatform);
+	const FName& PlatformPluginName = FName(*GetDesiredPluginName(EAudioPlugin::MODULATION, AudioPlatform));
+	const FName& PluginName = PlatformPluginName == NAME_None ? GetDefaultModulationPluginName() : PlatformPluginName;
+	const FName& FeatureName = IAudioModulationFactory::GetModularFeatureName();
 
-	TArray<IAudioModulationFactory *> ModulationPluginFactories = IModularFeatures::Get().GetModularFeatureImplementations<IAudioModulationFactory>(IAudioModulationFactory::GetModularFeatureName());
-
-	//Iterate through all of the plugins we've discovered:
-	for (IAudioModulationFactory* PluginFactory : ModulationPluginFactories)
+	TArray<IAudioModulationFactory*> Factories = IModularFeatures::Get().GetModularFeatureImplementations<IAudioModulationFactory>(FeatureName);
+	for (IAudioModulationFactory* Factory : Factories)
 	{
 		//if this plugin's name matches the name found in the platform settings, use it:
-		if (PluginFactory->GetDisplayName().Equals(DesiredModulationPlugin))
+		if (Factory->GetDisplayName() == PluginName)
 		{
-			return PluginFactory;
+			return Factory;
 		}
 	}
 
@@ -203,3 +213,15 @@ FString AudioPluginUtilities::GetDesiredPluginName(EAudioPlugin PluginType, EAud
 	GConfig->GetString(GetPlatformConfigSection(AudioPlatform), GetPluginConfigName(PluginType), PluginName, GEngineIni);
 	return PluginName;
 }
+
+const FName& AudioPluginUtilities::GetDefaultModulationPluginName()
+{
+	static FName DefaultModulationPluginName(TEXT("DefaultModulationPlugin"));
+
+	if (!DefaultModulationPluginCVar.IsEmpty())
+	{
+		DefaultModulationPluginName = FName(*DefaultModulationPluginCVar);
+	}
+
+	return DefaultModulationPluginName;
+}
\ No newline at end of file
diff --git a/Engine/Source/Runtime/Engine/Public/AudioPluginUtilities.h b/Engine/Source/Runtime/Engine/Public/AudioPluginUtilities.h
--- a/Engine/Source/Runtime/Engine/Public/AudioPluginUtilities.h
+++ b/Engine/Source/Runtime/Engine/Public/AudioPluginUtilities.h
@@ -38,6 +38,8 @@ struct ENGINE_API AudioPluginUtilities
 	/** Platform config section for each platform's target settings. */
 	static const TCHAR* GetPlatformConfigSection(EAudioPlatform AudioPlatform);
 
+	static const FName& GetDefaultModulationPluginName();
+
 	/*
 	 * These functions return a pointer to the plugin factory
 	 * that matches the plugin name specified in the target
diff --git a/Engine/Source/Runtime/Engine/Public/IAudioExtensionPlugin.h b/Engine/Source/Runtime/Engine/Public/IAudioExtensionPlugin.h
--- a/Engine/Source/Runtime/Engine/Public/IAudioExtensionPlugin.h
+++ b/Engine/Source/Runtime/Engine/Public/IAudioExtensionPlugin.h
@@ -476,7 +476,7 @@ public:
 /* IAudioModulationFactory                                              */
 /*                                                                      */
 /************************************************************************/
-class IAudioModulationFactory : public IAudioPluginFactory, public IModularFeature
+class IAudioModulationFactory : public IModularFeature
 {
 public:
 	/** Virtual destructor */
@@ -491,13 +491,7 @@ public:
 		return AudioExtFeatureName;
 	}
 
-	/* Begin IAudioPluginWithMetadata implementation */
-	virtual FString GetDisplayName() override
-	{
-		static FString DisplayName = FString(TEXT("Generic Audio Modulation Plugin"));
-		return DisplayName;
-	}
-	/* End IAudioPluginWithMetadata implementation */
+	virtual const FName& GetDisplayName() const = 0;
 
 	virtual TAudioModulationPtr CreateNewModulationPlugin(FAudioDevice* OwningDevice) = 0;
 
@@ -536,6 +530,11 @@ public:
 	/** Initialize the modulation plugin with the same rate and number of sources. */
 	virtual void Initialize(const FAudioPluginInitializationParams& InitializationParams) { }
 
+#if WITH_EDITOR
+	/** Allows modulation plugin to respond to settings updates when client is editing settings. */
+	virtual void OnEditSource(const USoundModulationPluginSourceSettingsBase& Settings) { }
+#endif // WITH_EDITOR
+
 	/** Called when a source is assigned to a voice. */
 	virtual void OnInitSource(const uint32 SourceId, const FName& AudioComponentUserId, const uint32 NumChannels, USoundModulationPluginSourceSettingsBase* Settings) { }
 
@@ -544,7 +543,7 @@ public:
 
 #if !UE_BUILD_SHIPPING
 	/** Request to post help from active plugin (non-shipping builds only) */
-	virtual bool OnPostHelp(FCommonViewportClient* ViewportClient, const TCHAR* Stream) { return false;  };
+	virtual bool OnPostHelp(FCommonViewportClient* ViewportClient, const TCHAR* Stream) { return false; };
 
 	/** Render stats pertaining to modulation (non-shipping builds only) */
 	virtual int32 OnRenderStat(FViewport* Viewport, FCanvas* Canvas, int32 X, int32 Y, const UFont& Font, const FVector* ViewLocation, const FRotator* ViewRotation) { return Y; }
