Remove code optimization from build cs

diff --git a/Engine/Plugins/Runtime/AudioModulation/Source/AudioModulation/AudioModulation.Build.cs b/Engine/Plugins/Runtime/AudioModulation/Source/AudioModulation/AudioModulation.Build.cs
--- a/Engine/Plugins/Runtime/AudioModulation/Source/AudioModulation/AudioModulation.Build.cs
+++ b/Engine/Plugins/Runtime/AudioModulation/Source/AudioModulation/AudioModulation.Build.cs
@@ -6,8 +6,6 @@ namespace UnrealBuildTool.Rules
 	{
 		public AudioModulation(ReadOnlyTargetRules Target) : base(Target)
 		{
-			OptimizeCode = CodeOptimization.Never;
-
 			PublicDependencyModuleNames.AddRange(
 				new string[]
 				{
