Move implementation of helper functions back to ModulationStatics and reference by Patch to avoid multiple definition compiler error in non-unity builds.

diff --git a/Engine/Plugins/Runtime/AudioModulation/Source/AudioModulation/Private/AudioModulationStatics.cpp b/Engine/Plugins/Runtime/AudioModulation/Source/AudioModulation/Private/AudioModulationStatics.cpp
--- a/Engine/Plugins/Runtime/AudioModulation/Source/AudioModulation/Private/AudioModulationStatics.cpp
+++ b/Engine/Plugins/Runtime/AudioModulation/Source/AudioModulation/Private/AudioModulationStatics.cpp
@@ -6,13 +6,39 @@
 #include "AudioModulation.h"
 #include "AudioModulationInternal.h"
 #include "AudioModulationLogging.h"
-#include "AudioModulationUtils.h"
+
 #include "Engine/Engine.h"
 #include "SoundModulatorBus.h"
 
 #define LOCTEXT_NAMESPACE "AudioModulationStatics"
 
 
+namespace
+{
+	template <class T>
+	T* CreateModulatorBus(const UObject* WorldContextObject, FName Name, float DefaultValue, bool Activate)
+	{
+		UWorld* World = GetAudioWorld(WorldContextObject);
+		if (!World)
+		{
+			return nullptr;
+		}
+
+		T* NewBus = NewObject<T>(nullptr, Name, RF_Transient);
+		NewBus->DefaultValue = DefaultValue;
+
+		if (Activate)
+		{
+			if (AudioModulation::FAudioModulationImpl* ModulationImpl = GetModulationImpl(World))
+			{
+				ModulationImpl->ActivateBus(*NewBus);
+			}
+		}
+
+		return NewBus;
+	}
+} // namespace <>
+
 //////////////////////////////////////////////////////////////////////////
 // UAudioModulationStatics
 
@@ -59,24 +85,69 @@ void UAudioModulationStatics::ActivateModulator(const UObject* WorldContextObjec
 	}
 }
 
+UWorld* UAudioModulationStatics::GetAudioWorld(const UObject* WorldContextObject)
+{
+	if (!GEngine || !GEngine->UseSound())
+	{
+		return nullptr;
+	}
+
+	UWorld* World = GEngine->GetWorldFromContextObject(WorldContextObject, EGetWorldErrorMode::LogAndReturnNull);
+	if (!World || !World->bAllowAudioPlayback || World->IsNetMode(NM_DedicatedServer))
+	{
+		return nullptr;
+	}
+
+	return World;
+}
+
+AudioModulation::FAudioModulationImpl* UAudioModulationStatics::GetModulationImpl(UWorld* World)
+{
+	FAudioDevice* AudioDevice = nullptr;
+	if (World)
+	{
+		AudioDevice = World->GetAudioDevice();
+	}
+	else
+	{
+		if (GEngine)
+		{
+			if (FAudioDeviceManager* DeviceManager = GEngine->GetAudioDeviceManager())
+			{
+				AudioDevice = DeviceManager->GetActiveAudioDevice();
+			}
+		}
+	}
+
+	if (AudioDevice && AudioDevice->IsModulationPluginEnabled())
+	{
+		if (IAudioModulation* ModulationInterface = AudioDevice->ModulationInterface.Get())
+		{
+			return static_cast<AudioModulation::FAudioModulation*>(ModulationInterface)->GetImpl();
+		}
+	}
+
+	return nullptr;
+}
+
 USoundVolumeModulatorBus* UAudioModulationStatics::CreateVolumeBus(const UObject* WorldContextObject, FName Name, float DefaultValue, bool Activate)
 {
-	return AudioModulation::CreateModulatorBus<USoundVolumeModulatorBus>(WorldContextObject, Name, DefaultValue, Activate);
+	return CreateModulatorBus<USoundVolumeModulatorBus>(WorldContextObject, Name, DefaultValue, Activate);
 }
 
 USoundPitchModulatorBus* UAudioModulationStatics::CreatePitchBus(const UObject* WorldContextObject, FName Name, float DefaultValue, bool Activate)
 {
-	return AudioModulation::CreateModulatorBus<USoundPitchModulatorBus>(WorldContextObject, Name, DefaultValue, Activate);
+	return CreateModulatorBus<USoundPitchModulatorBus>(WorldContextObject, Name, DefaultValue, Activate);
 }
 
 USoundLPFModulatorBus* UAudioModulationStatics::CreateLPFBus(const UObject* WorldContextObject, FName Name, float DefaultValue, bool Activate)
 {
-	return AudioModulation::CreateModulatorBus<USoundLPFModulatorBus>(WorldContextObject, Name, DefaultValue, Activate);
+	return CreateModulatorBus<USoundLPFModulatorBus>(WorldContextObject, Name, DefaultValue, Activate);
 }
 
 USoundHPFModulatorBus* UAudioModulationStatics::CreateHPFBus(const UObject* WorldContextObject, FName Name, float DefaultValue, bool Activate)
 {
-	return AudioModulation::CreateModulatorBus<USoundHPFModulatorBus>(WorldContextObject, Name, DefaultValue, Activate);
+	return CreateModulatorBus<USoundHPFModulatorBus>(WorldContextObject, Name, DefaultValue, Activate);
 }
 
 USoundModulatorLFO* UAudioModulationStatics::CreateLFO(const UObject* WorldContextObject, FName Name, float Amplitude, float Frequency, float Offset, bool Activate)
diff --git a/Engine/Plugins/Runtime/AudioModulation/Source/AudioModulation/Private/AudioModulationStatics.h b/Engine/Plugins/Runtime/AudioModulation/Source/AudioModulation/Private/AudioModulationStatics.h
--- a/Engine/Plugins/Runtime/AudioModulation/Source/AudioModulation/Private/AudioModulationStatics.h
+++ b/Engine/Plugins/Runtime/AudioModulation/Source/AudioModulation/Private/AudioModulationStatics.h
@@ -8,10 +8,10 @@
 
 #include "AudioModulationStatics.generated.h"
 
-
 class USoundModulatorBus;
 class USoundModulatorBusMix;
 
+
 UCLASS()
 class UAudioModulationStatics : public UBlueprintFunctionLibrary
 {
@@ -86,6 +86,16 @@ public:
 	UFUNCTION(BlueprintCallable, Category = "Audio", DisplayName = "Create Sound Modulation Bus Mix", meta = (WorldContext = "WorldContextObject", Keywords = "create bus mix modulation modulator"))
 	static USoundModulatorBusMix* CreateBusMix(const UObject* WorldContextObject, FName Name, TArray<USoundModulatorBusBase*> Buses, float TargetValue, bool Activate);
 
+	/**
+	 * Returns world associated with provided context object
+	 */
+	static UWorld* GetAudioWorld(const UObject* WorldContextObject);
+
+	/**
+	 * Returns modulation implementation associated with the provided world
+	 */
+	static AudioModulation::FAudioModulationImpl* GetModulationImpl(UWorld* World);
+
 	/** Deactivates a bus. Does nothing if an instance of the provided bus is already inactive
 	 * @param Bus - Scope of modulator
 	 */
diff --git a/Engine/Plugins/Runtime/AudioModulation/Source/AudioModulation/Private/AudioModulationUtils.h b/Engine/Plugins/Runtime/AudioModulation/Source/AudioModulation/Private/AudioModulationUtils.h
deleted file mode 100644
diff --git a/Engine/Plugins/Runtime/AudioModulation/Source/AudioModulation/Private/SoundModulationPatch.cpp b/Engine/Plugins/Runtime/AudioModulation/Source/AudioModulation/Private/SoundModulationPatch.cpp
--- a/Engine/Plugins/Runtime/AudioModulation/Source/AudioModulation/Private/SoundModulationPatch.cpp
+++ b/Engine/Plugins/Runtime/AudioModulation/Source/AudioModulation/Private/SoundModulationPatch.cpp
@@ -4,7 +4,7 @@
 #include "AudioDevice.h"
 #include "AudioModulation.h"
 #include "AudioModulationInternal.h"
-#include "AudioModulationUtils.h"
+#include "AudioModulationStatics.h"
 #include "SoundModulationTransform.h"
 
 
@@ -29,7 +29,7 @@ void USoundModulationSettings::PostEditChangeProperty(FPropertyChangedEvent& Pro
 		return;
 	}
 
-	if (AudioModulation::FAudioModulationImpl* Impl = AudioModulation::GetModulationImpl(GetWorld()))
+	if (AudioModulation::FAudioModulationImpl* Impl = UAudioModulationStatics::GetModulationImpl(GetWorld()))
 	{
 		Impl->OnEditSource(*this);
 	}
@@ -42,7 +42,7 @@ void USoundModulationSettings::PostEditChangeChainProperty(FPropertyChangedChain
 		return;
 	}
 
-	if (AudioModulation::FAudioModulationImpl* Impl = AudioModulation::GetModulationImpl(GetWorld()))
+	if (AudioModulation::FAudioModulationImpl* Impl = UAudioModulationStatics::GetModulationImpl(GetWorld()))
 	{
 		Impl->OnEditSource(*this);
 	}
