- Fix for engine unity build issues round 2 (missed compiler errors due to plugin disabled in local project)
- WIP: Make lifetime observed by auto-activate/deactivate for modulation types observe abstract sound's lifetime, not sound's source lifetime

diff --git a/Engine/Plugins/Runtime/AudioModulation/Source/AudioModulation/Private/AudioModulation.cpp b/Engine/Plugins/Runtime/AudioModulation/Source/AudioModulation/Private/AudioModulation.cpp
--- a/Engine/Plugins/Runtime/AudioModulation/Source/AudioModulation/Private/AudioModulation.cpp
+++ b/Engine/Plugins/Runtime/AudioModulation/Source/AudioModulation/Private/AudioModulation.cpp
@@ -38,11 +38,21 @@ namespace AudioModulation
 	}
 #endif // WITH_EDITOR
 
-	void FAudioModulation::OnInitSource(const uint32 SourceId, const FName& AudioComponentUserId, const uint32 NumChannels, USoundModulationPluginSourceSettingsBase* Settings)
+	void FAudioModulation::OnInitSound(const ModulationSoundId SoundId, const USoundModulationPluginSourceSettingsBase& Settings)
+	{
+		Impl->OnInitSound(SoundId, Settings);
+	}
+
+	void FAudioModulation::OnInitSource(const uint32 SourceId, const FName& AudioComponentUserId, const uint32 NumChannels, const USoundModulationPluginSourceSettingsBase& Settings)
 	{
 		Impl->OnInitSource(SourceId, AudioComponentUserId, NumChannels, Settings);
 	}
 
+	void FAudioModulation::OnReleaseSound(const ModulationSoundId SoundId, const USoundModulationPluginSourceSettingsBase& Settings)
+	{
+		Impl->OnReleaseSound(SoundId, Settings);
+	}
+
 	void FAudioModulation::OnReleaseSource(const uint32 SourceId)
 	{
 		Impl->OnReleaseSource(SourceId);
diff --git a/Engine/Plugins/Runtime/AudioModulation/Source/AudioModulation/Private/AudioModulationImpl.cpp b/Engine/Plugins/Runtime/AudioModulation/Source/AudioModulation/Private/AudioModulationImpl.cpp
--- a/Engine/Plugins/Runtime/AudioModulation/Source/AudioModulation/Private/AudioModulationImpl.cpp
+++ b/Engine/Plugins/Runtime/AudioModulation/Source/AudioModulation/Private/AudioModulationImpl.cpp
@@ -20,37 +20,38 @@ DECLARE_DWORD_COUNTER_STAT(TEXT("LFO Count"), STAT_AudioModulationLFOCount, STAT
 DECLARE_DWORD_COUNTER_STAT(TEXT("Mix Count"), STAT_AudioModulationMixCount, STATGROUP_AudioModulation)
 
 
-namespace AudioModulation
+namespace
 {
-	namespace
+	void MixInModulationValue(ESoundModulatorOperator& Operator, float ModStageValue, float& Value)
 	{
-		void MixInModulationValue(ESoundModulatorOperator& Operator, float ModStageValue, float& Value)
+		switch (Operator)
 		{
-			switch (Operator)
-			{
-				case ESoundModulatorOperator::Max:
-				{
-					Value = FMath::Max(ModStageValue, Value);
-				}
-				break;
+		case ESoundModulatorOperator::Max:
+		{
+			Value = FMath::Max(ModStageValue, Value);
+		}
+		break;
 
-				case ESoundModulatorOperator::Min:
-				{
-					Value = FMath::Min(ModStageValue, Value);
-				}
-				break;
+		case ESoundModulatorOperator::Min:
+		{
+			Value = FMath::Min(ModStageValue, Value);
+		}
+		break;
 
-				case ESoundModulatorOperator::Multiply:
-				default:
-				{
-					Value *= ModStageValue;
-					static_assert(static_cast<int32>(ESoundModulatorOperator::Count) == 3, "Possible missing operator switch case coverage");
-				}
-				break;
-			}
+		case ESoundModulatorOperator::Multiply:
+		default:
+		{
+			Value *= ModStageValue;
+			static_assert(static_cast<int32>(ESoundModulatorOperator::Count) == 3, "Possible missing operator switch case coverage");
+		}
+		break;
 		}
-	} // namespace <>
+	}
+} // namespace <>
 
+
+namespace AudioModulation
+{
 	void FAudioModulationImpl::Initialize(const FAudioPluginInitializationParams& InitializationParams)
 	{
 		SourceSettings.AddDefaulted(InitializationParams.NumSources);
@@ -77,18 +78,17 @@ namespace AudioModulation
 	}
 #endif // WITH_EDITOR
 
-	void FAudioModulationImpl::OnInitSource(const uint32 SourceId, const FName& AudioComponentUserId, const uint32 NumChannels, USoundModulationPluginSourceSettingsBase* InSettings)
+	void FAudioModulationImpl::OnInitSound(const ModulationSoundId SoundId, const USoundModulationPluginSourceSettingsBase& InSettings)
 	{
 		check(IsInAudioThread());
 
-		if (USoundModulationSettings* Settings = CastChecked<USoundModulationSettings>(InSettings))
+		if (const USoundModulationSettings* Settings = CastChecked<USoundModulationSettings>(&InSettings))
 		{
-			SourceSettings[SourceId] = FModulationSettingsProxy(*Settings);
 			for (const USoundModulatorBusMix* Mix : Settings->Mixes)
 			{
 				if (Mix)
 				{
-					ActivateBusMix(*Mix, /* bReset */ false);
+					ActivateBusMix(*Mix, /* bReset */ false, SoundId);
 				}
 			}
 
@@ -139,6 +139,16 @@ namespace AudioModulation
 		}
 	}
 
+	void FAudioModulationImpl::OnInitSource(const uint32 SourceId, const FName& AudioComponentUserId, const uint32 NumChannels, const USoundModulationPluginSourceSettingsBase& InSettings)
+	{
+		check(IsInAudioThread());
+
+		if (const USoundModulationSettings* Settings = CastChecked<USoundModulationSettings>(&InSettings))
+		{
+			SourceSettings[SourceId] = FModulationSettingsProxy(*Settings);
+		}
+	}
+
 	void FAudioModulationImpl::OnReleaseSource(const uint32 SourceId)
 	{
 		check(IsInAudioThread());
@@ -146,6 +156,12 @@ namespace AudioModulation
 		SourceSettings[SourceId] = FModulationSettingsProxy();
 	}
 
+	void FAudioModulationImpl::OnReleaseSound(const ModulationSoundId SoundId, const USoundModulationPluginSourceSettingsBase& Settings)
+	{
+		check(IsInAudioThread());
+
+	}
+
 #if !UE_BUILD_SHIPPING
 	bool FAudioModulationImpl::OnPostHelp(FCommonViewportClient* ViewportClient, const TCHAR* Stream)
 	{
@@ -163,38 +179,40 @@ namespace AudioModulation
 	}
 #endif // !UE_BUILD_SHIPPING
 
-	void FAudioModulationImpl::ActivateBus(const USoundModulatorBusBase& Bus)
+	void FAudioModulationImpl::ActivateBus(const USoundModulatorBusBase& Bus, const ModulationSoundId SoundId)
 	{
-		FModulatorBusProxy BusProxy(Bus);
-		ActivateBus(BusProxy);
-	}
+		const FModulatorBusProxy NewBusProxy(Bus);
 
-	void FAudioModulationImpl::ActivateBus(const FModulatorBusProxy& BusProxy)
-	{
-		auto ActivateBusInternal = [this, BusProxy]()
+		auto ActivateBusInternal = [this, NewBusProxy, SoundId]()
 		{
-			const BusId NewBusId = BusProxy.GetBusId();
-			if (!ActiveBuses.Contains(NewBusId))
+			const BusId NewBusId = NewBusProxy.GetBusId();
+			FModulatorBusProxy* BusProxy = ActiveBuses.Find(NewBusId);
+			if (!BusProxy)
+			{
+				BusProxy = &ActiveBuses.Add(NewBusId, NewBusProxy);
+			}
+
+			if (BusProxy->GetAutoActivate() && SoundId != INDEX_NONE)
 			{
-				ActiveBuses.Add(NewBusId, BusProxy);
+				BusProxy->IncRefSound();
 			}
 		};
 
 		IsInAudioThread() ? ActivateBusInternal() : FAudioThread::RunCommandOnAudioThread(ActivateBusInternal);
 	}
 
-	void FAudioModulationImpl::ActivateBusMix(const USoundModulatorBusMix& BusMix, bool bReset)
+	void FAudioModulationImpl::ActivateBusMix(const USoundModulatorBusMix& BusMix, bool bReset, const ModulationSoundId SoundId)
 	{
 		const BusMixId MixId = static_cast<BusMixId>(BusMix.GetUniqueID());
 		FModulatorBusMixProxy MixProxy(BusMix);
 
-		auto ActivateMixInternal = [this, MixId, MixProxy, bReset]()
+		auto ActivateMixInternal = [this, MixId, MixProxy, bReset, SoundId]()
 		{
 			FModulatorBusMixProxy* BusMixProxy = nullptr;
 			if (FModulatorBusMixProxy* ExistingMixProxy = ActiveBusMixes.Find(MixId))
 			{
-				// Reactivate just in case it is currently stopping but not yet stopped.
-				ExistingMixProxy->SetActive();
+				// Enable in case mix is currently stopping but not yet stopped.
+				ExistingMixProxy->SetEnabled();
 				if (bReset)
 				{
 					for (int32 i = 0; i < ExistingMixProxy->Channels.Num(); ++i)
@@ -212,21 +230,32 @@ namespace AudioModulation
 			{
 				BusMixProxy = &ActiveBusMixes.Emplace(MixId, MixProxy);
 			}
+
+			if (BusMixProxy->GetAutoActivate() && SoundId != INDEX_NONE)
+			{
+				BusMixProxy->IncRefSound();
+			}
 		};
 
 		IsInAudioThread() ? ActivateMixInternal() : FAudioThread::RunCommandOnAudioThread(ActivateMixInternal);
 	}
 
-	void FAudioModulationImpl::ActivateLFO(const USoundModulatorLFO& LFO)
+	void FAudioModulationImpl::ActivateLFO(const USoundModulatorLFO& LFO, const ModulationSoundId SoundId)
 	{
 		const AudioModulation::LFOId LFOId = static_cast<const AudioModulation::LFOId>(LFO.GetUniqueID());
-		FModulatorLFOProxy LFOProxy(LFO);
+		const FModulatorLFOProxy NewLFOProxy(LFO);
 
-		auto ActivateLFOInternal = [this, LFOId, LFOProxy]()
+		auto ActivateLFOInternal = [this, LFOId, NewLFOProxy, SoundId]()
 		{
-			if (!ActiveLFOs.Contains(LFOId))
+			FModulatorLFOProxy* LFOProxy = nullptr;
+			if (!LFOProxy)
+			{
+				LFOProxy = &ActiveLFOs.Add(LFOId, NewLFOProxy);
+			}
+
+			if (SoundId != INDEX_NONE)
 			{
-				ActiveLFOs.Add(LFOId, LFOProxy);
+				LFOProxy->IncRefSound();
 			}
 		};
 
@@ -273,41 +302,61 @@ namespace AudioModulation
 		return Value;
 	}
 
-	void FAudioModulationImpl::DeactivateBusMix(const BusMixId BusMixId)
+	void FAudioModulationImpl::DeactivateBusMix(const BusMixId BusMixId, const ModulationSoundId SoundId)
 	{
-		auto DeactivateMix = [this, BusMixId]()
+		auto DeactivateMix = [this, BusMixId, SoundId]()
 		{
 			check(IsInAudioThread());
 
 			if (FModulatorBusMixProxy* Mix = ActiveBusMixes.Find(BusMixId))
 			{
-				Mix->SetStopping();
+				if (SoundId == INDEX_NONE)
+				{
+					Mix->SetStopping();
+				}
+				else if (Mix->GetAutoActivate())
+				{
+					if (Mix->DecRefSound() == 0)
+					{
+						Mix->SetStopping();
+					}
+				}
 			}
 		};
 
 		IsInAudioThread() ? DeactivateMix() : FAudioThread::RunCommandOnAudioThread(DeactivateMix);
 	}
 
-	void FAudioModulationImpl::DeactivateBus(const BusId BusId)
+	void FAudioModulationImpl::DeactivateBus(const BusId BusId, const ModulationSoundId SoundId)
 	{
-		if (!IsInAudioThread())
+		auto Deactivate = [this, BusId, SoundId]()
 		{
-			FAudioThread::RunCommandOnAudioThread([this, BusId]()
-			{
-				DeactivateBus(BusId);
-			});
+			check(IsInAudioThread());
 
-			return;
-		}
+			if (FModulatorBusProxy* Bus = ActiveBuses.Find(BusId))
+			{
+				if (SoundId == INDEX_NONE)
+				{
+					ActiveBuses.Remove(BusId);
+				}
+				else if (Bus->GetAutoActivate())
+				{
+					if (Bus->DecRefSound() == 0)
+					{
+						ActiveBuses.Remove(BusId);
+					}
+				}
+			}
+		};
 
-		ActiveBuses.Remove(BusId);
+		IsInAudioThread() ? Deactivate() : FAudioThread::RunCommandOnAudioThread(Deactivate);
 	}
 
-	void FAudioModulationImpl::DeactivateLFO(const LFOId LFOId)
+	void FAudioModulationImpl::DeactivateLFO(const LFOId LFOId, const ModulationSoundId SoundId)
 	{
 		if (!IsInAudioThread())
 		{
-			FAudioThread::RunCommandOnAudioThread([this, LFOId]()
+			FAudioThread::RunCommandOnAudioThread([this, LFOId, SoundId]()
 			{
 				DeactivateLFO(LFOId);
 			});
@@ -318,38 +367,6 @@ namespace AudioModulation
 		ActiveLFOs.Remove(LFOId);
 	}
 
-	TSet<BusId> FAudioModulationImpl::GetReferencedBusIds() const
-	{
-		auto GetActivePatchBusIds = [this](const FModulationPatchProxy& Patch, TSet<BusId>& OutBusIds)
-		{
-			for (const FModulationInputProxy& InputProxy : Patch.InputProxies)
-			{
-				if (InputProxy.bSampleAndHold)
-				{
-					if (!Patch.OutputProxy.bInitialized)
-					{
-						OutBusIds.Add(InputProxy.BusId);
-					}
-				}
-				else
-				{
-					OutBusIds.Add(InputProxy.BusId);
-				}
-			}
-		};
-
-		TSet<BusId> BusIds;
-		for (const FModulationSettingsProxy& Settings : SourceSettings)
-		{
-			GetActivePatchBusIds(Settings.Highpass, BusIds);
-			GetActivePatchBusIds(Settings.Lowpass, BusIds);
-			GetActivePatchBusIds(Settings.Pitch, BusIds);
-			GetActivePatchBusIds(Settings.Volume, BusIds);
-		}
-
-		return MoveTemp(BusIds);
-	}
-
 	bool FAudioModulationImpl::IsBusActive(const BusId BusId) const
 	{
 		check(IsInAudioThread());
@@ -376,10 +393,6 @@ namespace AudioModulation
 		// Update LFOs (prior to bus mixing to avoid single-frame latency)
 		for (TPair<LFOId, FModulatorLFOProxy>& Pair : ActiveLFOs)
 		{
-			// Clear is active as bus calls to MixLFO in subsequent loop will
-			// mark as active.  If no buses are currently using the LFO and
-			// its set to 'Limited Lifetime', then it will be marked to destroy.
-			Pair.Value.ClearIsActive();
 			Pair.Value.Update(Elapsed);
 		}
 
@@ -396,45 +409,13 @@ namespace AudioModulation
 		for (TPair<BusMixId, FModulatorBusMixProxy>& Pair : ActiveBusMixes)
 		{
 			Pair.Value.Update(Elapsed, ActiveBuses);
-			if (Pair.Value.CanDeactivate())
+			if (Pair.Value.CanDestroy())
 			{
 				UE_LOG(LogAudioModulation, Log, TEXT("Audio modulation mix '%u' stopped."), static_cast<uint32>(Pair.Key));
 				MixesToDeactivate.Add(Pair.Key);
 			}
 		}
 
-		{	// Destroy LFO proxies ready to deactivate
-			TArray<LFOId> LFOsToDeactivate;
-			for (TPair<LFOId, FModulatorLFOProxy>& Pair : ActiveLFOs)
-			{
-				if (Pair.Value.CanDeactivate())
-				{
-					LFOsToDeactivate.Add(Pair.Key);
-				}
-			}
-			for (const LFOId& LFOId : LFOsToDeactivate)
-			{
-				DeactivateLFO(LFOId);
-			}
-		}
-
-		{	// Destroy Bus proxies ready to deactivate
-			const TSet<BusId> ReferencedBuses = GetReferencedBusIds();
-			TArray<BusId> BusesToDeactivate;
-			for (TPair<BusId, FModulatorBusProxy>& Pair : ActiveBuses)
-			{
-				FModulatorBusProxy& BusProxy = Pair.Value;
-				if (BusProxy.CanDeactivate() && !ReferencedBuses.Contains(BusProxy.GetBusId()))
-				{
-					BusesToDeactivate.Add(Pair.Key);
-				}
-			}
-			for (const BusId& BusId : BusesToDeactivate)
-			{
-				DeactivateBus(BusId);
-			}
-		}
-
 		// Destroy mixes that have stopped
 		for (const BusMixId& MixId : MixesToDeactivate)
 		{
diff --git a/Engine/Plugins/Runtime/AudioModulation/Source/AudioModulation/Private/AudioModulationInternal.h b/Engine/Plugins/Runtime/AudioModulation/Source/AudioModulation/Private/AudioModulationInternal.h
--- a/Engine/Plugins/Runtime/AudioModulation/Source/AudioModulation/Private/AudioModulationInternal.h
+++ b/Engine/Plugins/Runtime/AudioModulation/Source/AudioModulation/Private/AudioModulationInternal.h
@@ -29,16 +29,17 @@ namespace AudioModulation
 		void OnEditSource(const USoundModulationPluginSourceSettingsBase& Settings);
 #endif // WITH_EDITOR
 
-		void OnInitSource(const uint32 SourceId, const FName& AudioComponentUserId, const uint32 NumChannels, USoundModulationPluginSourceSettingsBase* Settings);
+		void OnInitSound(const ModulationSoundId SoundId, const USoundModulationPluginSourceSettingsBase& Settings);
+		void OnInitSource(const uint32 SourceId, const FName& AudioComponentUserId, const uint32 NumChannels, const USoundModulationPluginSourceSettingsBase& Settings);
+		void OnReleaseSound(const ModulationSoundId SoundId, const USoundModulationPluginSourceSettingsBase& Settings);
 		void OnReleaseSource(const uint32 SourceId);
 
-		void ActivateBus(const USoundModulatorBusBase& Bus);
-		void ActivateBus(const FModulatorBusProxy& BusProxy);
-		void ActivateBusMix(const USoundModulatorBusMix& BusMix, bool bReset);
-		void ActivateLFO(const USoundModulatorLFO& LFO);
-		void DeactivateBusMix(const BusMixId BusMixId);
-		void DeactivateBus(const BusId BusId);
-		void DeactivateLFO(const LFOId LFOId);
+		void ActivateBus(const USoundModulatorBusBase& Bus, const ModulationSoundId SoundId = INDEX_NONE);
+		void ActivateBusMix(const USoundModulatorBusMix& BusMix, bool bReset, const ModulationSoundId SoundId = INDEX_NONE);
+		void ActivateLFO(const USoundModulatorLFO& LFO, const ModulationSoundId SoundId = INDEX_NONE);
+		void DeactivateBusMix(const BusMixId BusMixId, const ModulationSoundId SoundId = INDEX_NONE);
+		void DeactivateBus(const BusId BusId, const ModulationSoundId SoundId = INDEX_NONE);
+		void DeactivateLFO(const LFOId LFOId, const ModulationSoundId SoundId = INDEX_NONE);
 
 		bool IsBusActive(const BusId BusId) const;
 
@@ -53,7 +54,6 @@ namespace AudioModulation
 
 	private:
 		float CalculateModulationValue(FModulationPatchProxy& Proxy) const;
-		TSet<BusId> GetReferencedBusIds() const;
 
 		BusMixProxyMap ActiveBusMixes;
 		BusProxyMap    ActiveBuses;
@@ -87,7 +87,9 @@ namespace AudioModulation
 		void OnEditSource(const USoundModulationPluginSourceSettingsBase& Settings) { }
 #endif // WITH_EDITOR
 
-		void OnInitSource(const uint32 SourceId, const FName& AudioComponentUserId, const uint32 NumChannels, USoundModulationPluginSourceSettingsBase* Settings) { }
+		void OnInitSound(const ModulationSoundId SoundId, const USoundModulationPluginSourceSettingsBase& Settings) { }
+		void OnInitSource(const uint32 SourceId, const FName& AudioComponentUserId, const uint32 NumChannels, const USoundModulationPluginSourceSettingsBase& Settings) { }
+		void OnReleaseSound(const ModulationSoundId SoundId, const USoundModulationPluginSourceSettingsBase& Settings) { }
 		void OnReleaseSource(const uint32 SourceId) { }
 
 #if !UE_BUILD_SHIPPING
@@ -96,12 +98,12 @@ namespace AudioModulation
 		bool OnToggleStat(FCommonViewportClient* ViewportClient, const TCHAR* Stream) { return false; }
 #endif // !UE_BUILD_SHIPPING
 
-		void ActivateBus(const USoundModulatorBusBase& Bus) { }
-		void ActivateBusMix(const USoundModulatorBusMix& BusMix, bool bReset) { }
-		void ActivateLFO(const USoundModulatorLFO& LFO) { }
-		void DeactivateBusMix(const BusMixId BusMixId) { }
-		void DeactivateBus(const BusId BusId) { }
-		void DeactivateLFO(const LFOId LFOId) { }
+		void ActivateBus(const USoundModulatorBusBase& Bus, const ModulationSoundId SoundId = INDEX_NONE) { }
+		void ActivateBusMix(const USoundModulatorBusMix& BusMix, bool bReset, const ModulationSoundId SoundId = INDEX_NONE) { }
+		void ActivateLFO(const USoundModulatorLFO& LFO, const ModulationSoundId SoundId = INDEX_NONE) { }
+		void DeactivateBusMix(const BusMixId BusMixId, const ModulationSoundId SoundId = INDEX_NONE) { }
+		void DeactivateBus(const BusId BusId, const ModulationSoundId SoundId = INDEX_NONE) { }
+		void DeactivateLFO(const LFOId LFOId, const ModulationSoundId SoundId = INDEX_NONE) { }
 
 		void ProcessAudio(const FAudioPluginSourceInputData& InputData, FAudioPluginSourceOutputData& OutputData) { }
 		void ProcessControls(const uint32 SourceId, FSoundModulationControls& Controls) { }
diff --git a/Engine/Plugins/Runtime/AudioModulation/Source/AudioModulation/Private/AudioModulationStatics.cpp b/Engine/Plugins/Runtime/AudioModulation/Source/AudioModulation/Private/AudioModulationStatics.cpp
--- a/Engine/Plugins/Runtime/AudioModulation/Source/AudioModulation/Private/AudioModulationStatics.cpp
+++ b/Engine/Plugins/Runtime/AudioModulation/Source/AudioModulation/Private/AudioModulationStatics.cpp
@@ -18,7 +18,7 @@ namespace
 	template <class T>
 	T* CreateModulatorBus(const UObject* WorldContextObject, FName Name, float DefaultValue, bool Activate)
 	{
-		UWorld* World = GetAudioWorld(WorldContextObject);
+		UWorld* World = UAudioModulationStatics::GetAudioWorld(WorldContextObject);
 		if (!World)
 		{
 			return nullptr;
@@ -29,7 +29,7 @@ namespace
 
 		if (Activate)
 		{
-			if (AudioModulation::FAudioModulationImpl* ModulationImpl = GetModulationImpl(World))
+			if (AudioModulation::FAudioModulationImpl* ModulationImpl = UAudioModulationStatics::GetModulationImpl(World))
 			{
 				ModulationImpl->ActivateBus(*NewBus);
 			}
@@ -54,8 +54,8 @@ void UAudioModulationStatics::ActivateBus(const UObject* WorldContextObject, USo
 		return;
 	}
 
-	UWorld* World = AudioModulation::GetAudioWorld(WorldContextObject);
-	if (AudioModulation::FAudioModulationImpl* ModulationImpl = AudioModulation::GetModulationImpl(World))
+	UWorld* World = GetAudioWorld(WorldContextObject);
+	if (AudioModulation::FAudioModulationImpl* ModulationImpl = GetModulationImpl(World))
 	{
 		ModulationImpl->ActivateBus(*Bus);
 	}
@@ -65,8 +65,8 @@ void UAudioModulationStatics::ActivateBusMix(const UObject* WorldContextObject,
 {
 	if (BusMix)
 	{
-		UWorld* World = AudioModulation::GetAudioWorld(WorldContextObject);
-		if (AudioModulation::FAudioModulationImpl* ModulationImpl = AudioModulation::GetModulationImpl(World))
+		UWorld* World = GetAudioWorld(WorldContextObject);
+		if (AudioModulation::FAudioModulationImpl* ModulationImpl = GetModulationImpl(World))
 		{
 			ModulationImpl->ActivateBusMix(*BusMix, false);
 		}
@@ -75,8 +75,8 @@ void UAudioModulationStatics::ActivateBusMix(const UObject* WorldContextObject,
 
 void UAudioModulationStatics::ActivateModulator(const UObject* WorldContextObject, USoundModulatorBase* Modulator)
 {
-	UWorld* World = AudioModulation::GetAudioWorld(WorldContextObject);
-	if (AudioModulation::FAudioModulationImpl* ModulationImpl = AudioModulation::GetModulationImpl(World))
+	UWorld* World = GetAudioWorld(WorldContextObject);
+	if (AudioModulation::FAudioModulationImpl* ModulationImpl = GetModulationImpl(World))
 	{
 		if (USoundModulatorLFO* LFO = Cast<USoundModulatorLFO>(Modulator))
 		{
@@ -152,7 +152,7 @@ USoundHPFModulatorBus* UAudioModulationStatics::CreateHPFBus(const UObject* Worl
 
 USoundModulatorLFO* UAudioModulationStatics::CreateLFO(const UObject* WorldContextObject, FName Name, float Amplitude, float Frequency, float Offset, bool Activate)
 {
-	UWorld* World = AudioModulation::GetAudioWorld(WorldContextObject);
+	UWorld* World = GetAudioWorld(WorldContextObject);
 	if (!World)
 	{
 		return nullptr;
@@ -165,7 +165,7 @@ USoundModulatorLFO* UAudioModulationStatics::CreateLFO(const UObject* WorldConte
 
 	if (Activate)
 	{
-		if (AudioModulation::FAudioModulationImpl* ModulationImpl = AudioModulation::GetModulationImpl(World))
+		if (AudioModulation::FAudioModulationImpl* ModulationImpl = GetModulationImpl(World))
 		{
 			ModulationImpl->ActivateLFO(*NewLFO);
 		}
@@ -176,7 +176,7 @@ USoundModulatorLFO* UAudioModulationStatics::CreateLFO(const UObject* WorldConte
 
 USoundModulatorBusMix* UAudioModulationStatics::CreateBusMix(const UObject* WorldContextObject, FName Name, TArray<USoundModulatorBusBase*> Buses, float TargetValue, bool Activate)
 {
-	UWorld* World = AudioModulation::GetAudioWorld(WorldContextObject);
+	UWorld* World = GetAudioWorld(WorldContextObject);
 	if (!World)
 	{
 		return nullptr;
@@ -199,7 +199,7 @@ USoundModulatorBusMix* UAudioModulationStatics::CreateBusMix(const UObject* Worl
 
 	if (Activate)
 	{
-		if (AudioModulation::FAudioModulationImpl* ModulationImpl = AudioModulation::GetModulationImpl(World))
+		if (AudioModulation::FAudioModulationImpl* ModulationImpl = GetModulationImpl(World))
 		{
 			ModulationImpl->ActivateBusMix(*NewBusMix, false);
 		}
@@ -215,8 +215,8 @@ void UAudioModulationStatics::DeactivateBus(const UObject* WorldContextObject, U
 		return;
 	}
 
-	UWorld* World = AudioModulation::GetAudioWorld(WorldContextObject);
-	if (AudioModulation::FAudioModulationImpl* ModulationImpl = AudioModulation::GetModulationImpl(World))
+	UWorld* World = GetAudioWorld(WorldContextObject);
+	if (AudioModulation::FAudioModulationImpl* ModulationImpl = GetModulationImpl(World))
 	{
 		auto BusId = static_cast<const AudioModulation::BusId>(Bus->GetUniqueID());
 		ModulationImpl->DeactivateBus(BusId);
@@ -227,8 +227,8 @@ void UAudioModulationStatics::DeactivateBusMix(const UObject* WorldContextObject
 {
 	if (BusMix)
 	{
-		UWorld* World = AudioModulation::GetAudioWorld(WorldContextObject);
-		if (AudioModulation::FAudioModulationImpl* ModulationImpl = AudioModulation::GetModulationImpl(World))
+		UWorld* World = GetAudioWorld(WorldContextObject);
+		if (AudioModulation::FAudioModulationImpl* ModulationImpl = GetModulationImpl(World))
 		{
 			auto BusMixId = static_cast<const AudioModulation::BusMixId>(BusMix->GetUniqueID());
 			ModulationImpl->DeactivateBusMix(BusMixId);
@@ -238,8 +238,8 @@ void UAudioModulationStatics::DeactivateBusMix(const UObject* WorldContextObject
 
 void UAudioModulationStatics::DeactivateModulator(const UObject* WorldContextObject, USoundModulatorBase* Modulator)
 {
-	UWorld* World = AudioModulation::GetAudioWorld(WorldContextObject);
-	if (AudioModulation::FAudioModulationImpl* ModulationImpl = AudioModulation::GetModulationImpl(World))
+	UWorld* World = GetAudioWorld(WorldContextObject);
+	if (AudioModulation::FAudioModulationImpl* ModulationImpl = GetModulationImpl(World))
 	{
 		if (USoundModulatorLFO* LFO = Cast<USoundModulatorLFO>(Modulator))
 		{
@@ -253,8 +253,8 @@ void UAudioModulationStatics::SetBusDefault(const UObject* WorldContextObject, U
 {
 	if (Bus)
 	{
-		UWorld* World = AudioModulation::GetAudioWorld(WorldContextObject);
-		if (AudioModulation::FAudioModulationImpl* ModulationImpl = AudioModulation::GetModulationImpl(World))
+		UWorld* World = GetAudioWorld(WorldContextObject);
+		if (AudioModulation::FAudioModulationImpl* ModulationImpl = GetModulationImpl(World))
 		{
 			ModulationImpl->SetBusDefault(*Bus, Value);
 		}
@@ -269,8 +269,8 @@ void UAudioModulationStatics::SetBusMixChannel(const UObject* WorldContextObject
 {
 	if (BusMix && Bus)
 	{
-		UWorld* World = AudioModulation::GetAudioWorld(WorldContextObject);
-		if (AudioModulation::FAudioModulationImpl* ModulationImpl = AudioModulation::GetModulationImpl(World))
+		UWorld* World = GetAudioWorld(WorldContextObject);
+		if (AudioModulation::FAudioModulationImpl* ModulationImpl = GetModulationImpl(World))
 		{
 			ModulationImpl->SetBusMixChannel(*BusMix, *Bus, TargetValue);
 		}
@@ -285,8 +285,8 @@ void UAudioModulationStatics::SetLFOFrequency(const UObject* WorldContextObject,
 {
 	if (LFO)
 	{
-		UWorld* World = AudioModulation::GetAudioWorld(WorldContextObject);
-		if (AudioModulation::FAudioModulationImpl* ModulationImpl = AudioModulation::GetModulationImpl(World))
+		UWorld* World = GetAudioWorld(WorldContextObject);
+		if (AudioModulation::FAudioModulationImpl* ModulationImpl = GetModulationImpl(World))
 		{
 			ModulationImpl->SetLFOFrequency(*LFO, Freq);
 		}
diff --git a/Engine/Plugins/Runtime/AudioModulation/Source/AudioModulation/Private/AudioModulationStatics.h b/Engine/Plugins/Runtime/AudioModulation/Source/AudioModulation/Private/AudioModulationStatics.h
--- a/Engine/Plugins/Runtime/AudioModulation/Source/AudioModulation/Private/AudioModulationStatics.h
+++ b/Engine/Plugins/Runtime/AudioModulation/Source/AudioModulation/Private/AudioModulationStatics.h
@@ -11,6 +11,11 @@
 class USoundModulatorBus;
 class USoundModulatorBusMix;
 
+namespace AudioModulation
+{
+	class FAudioModulationImpl;
+} // namespace AudioModulation
+
 
 UCLASS()
 class UAudioModulationStatics : public UBlueprintFunctionLibrary
diff --git a/Engine/Plugins/Runtime/AudioModulation/Source/AudioModulation/Private/SoundModulatorBus.cpp b/Engine/Plugins/Runtime/AudioModulation/Source/AudioModulation/Private/SoundModulatorBus.cpp
--- a/Engine/Plugins/Runtime/AudioModulation/Source/AudioModulation/Private/SoundModulatorBus.cpp
+++ b/Engine/Plugins/Runtime/AudioModulation/Source/AudioModulation/Private/SoundModulatorBus.cpp
@@ -13,7 +13,6 @@
 USoundModulatorBusBase::USoundModulatorBusBase(const FObjectInitializer& ObjectInitializer)
 	: Super(ObjectInitializer)
 	, bAutoActivate(0)
-	, bAutoDeactivate(0)
 	, DefaultValue(1.0f)
 	, Min(0.0f)
 	, Max(1.0f)
@@ -119,7 +118,8 @@ namespace AudioModulation
 		, MixValue(NAN)
 		, Operator(ESoundModulatorOperator::Multiply)
 		, Range(0.0f, 1.0f)
-		, bAutoDeactivate(false) 
+		, bAutoActivate(false)
+		, SoundRefCount(0)
 	{
 	}
 
@@ -133,7 +133,8 @@ namespace AudioModulation
 		, MixValue(NAN)
 		, Operator(Bus.GetOperator())
 		, Range(Bus.Min, Bus.Max)
-		, bAutoDeactivate(Bus.bAutoDeactivate)
+		, bAutoActivate(Bus.bAutoActivate)
+		, SoundRefCount(0)
 	{
 		if (Bus.Min > Bus.Max)
 		{
@@ -153,9 +154,9 @@ namespace AudioModulation
 		}
 	}
 
-	bool FModulatorBusProxy::CanDeactivate() const
+	bool FModulatorBusProxy::GetAutoActivate() const
 	{
-		return bAutoDeactivate;
+		return bAutoActivate;
 	}
 
 	AudioModulation::BusId FModulatorBusProxy::GetBusId() const
@@ -251,7 +252,6 @@ namespace AudioModulation
 		{
 			if (FModulatorLFOProxy* LFOProxy = LFOMap.Find(LFOId))
 			{
-				LFOProxy->SetIsActive();
 				LFOValue *= LFOProxy->GetValue();
 			}
 		}
@@ -272,4 +272,15 @@ namespace AudioModulation
 	{
 		Range = InRange;
 	}
+
+	int32 FModulatorBusProxy::DecRefSound()
+	{
+		check(SoundRefCount > 0);
+		return SoundRefCount--;
+	}
+
+	int32 FModulatorBusProxy::IncRefSound()
+	{
+		return SoundRefCount++;
+	}
 } // namespace AudioModulation
diff --git a/Engine/Plugins/Runtime/AudioModulation/Source/AudioModulation/Private/SoundModulatorBus.h b/Engine/Plugins/Runtime/AudioModulation/Source/AudioModulation/Private/SoundModulatorBus.h
--- a/Engine/Plugins/Runtime/AudioModulation/Source/AudioModulation/Private/SoundModulatorBus.h
+++ b/Engine/Plugins/Runtime/AudioModulation/Source/AudioModulation/Private/SoundModulatorBus.h
@@ -33,14 +33,10 @@ class USoundModulatorBusBase : public UObject
 	GENERATED_UCLASS_BODY()
 
 public:
-	/** Automatically activates bus when no sounds are being updated by it. */
+	/** Automatically activates/deactivates bus when sounds referencing asset are playing. */
 	UPROPERTY(EditAnywhere, Category = Modulation, BlueprintReadOnly)
 	uint8 bAutoActivate : 1;
 
-	/** Automatically deactivates bus when no sounds are being updated by it. */
-	UPROPERTY(EditAnywhere, Category = Modulation, BlueprintReadOnly)
-	uint8 bAutoDeactivate : 1;
-
 	/** Default value of modulator (when no mix is applied). */
 	UPROPERTY(EditAnywhere, Category = Modulation, BlueprintReadWrite, meta = (UIMin = "0", UIMax = "1"))
 	float DefaultValue;
@@ -103,8 +99,7 @@ namespace AudioModulation
 		FModulatorBusProxy();
 		FModulatorBusProxy(const USoundModulatorBusBase& Bus);
 
-		bool CanDeactivate() const;
-
+		bool GetAutoActivate() const;
 		AudioModulation::BusId GetBusId() const;
 		float GetDefaultValue() const;
 		float GetLFOValue() const;
@@ -118,6 +113,9 @@ namespace AudioModulation
 		void SetDefaultValue(const float Value);
 		void SetRange(const FVector2D& Range);
 
+		int32 IncRefSound();
+		int32 DecRefSound();
+
 #if !UE_BUILD_SHIPPING
 		const FString& GetName() const;
 #endif // !UE_BUILD_SHIPPING
@@ -142,7 +140,8 @@ namespace AudioModulation
 		ESoundModulatorOperator Operator;
 		FVector2D Range;
 
-		uint8 bAutoDeactivate : 1;
+		uint8 bAutoActivate : 1;
+		int32 SoundRefCount;
 	};
 	using BusProxyMap = TMap<BusId, FModulatorBusProxy>;
 } // namespace AudioModulation
diff --git a/Engine/Plugins/Runtime/AudioModulation/Source/AudioModulation/Private/SoundModulatorBusMix.cpp b/Engine/Plugins/Runtime/AudioModulation/Source/AudioModulation/Private/SoundModulatorBusMix.cpp
--- a/Engine/Plugins/Runtime/AudioModulation/Source/AudioModulation/Private/SoundModulatorBusMix.cpp
+++ b/Engine/Plugins/Runtime/AudioModulation/Source/AudioModulation/Private/SoundModulatorBusMix.cpp
@@ -65,12 +65,12 @@ namespace AudioModulation
 	}
 
 	FModulatorBusMixProxy::FModulatorBusMixProxy(const USoundModulatorBusMix& Mix)
+		: Status(BusMixStatus::Enabled)
+		, SoundRefCount(0)
+		, bAutoActivate(0)
 #if UE_BUILD_SHIPPING
-		: Status(BusMixStatus::Active)
-#else
-		: Name(Mix.GetName())
-		, Status(BusMixStatus::Active)
-#endif // !UE_BUILD_SHIPPING
+		, Name(Mix.GetName())
+#endif // UE_BUILD_SHIPPING
 	{
 		for (const FSoundModulatorBusMixChannel& Channel : Mix.Channels)
 		{
@@ -84,7 +84,7 @@ namespace AudioModulation
 						TEXT("USoundModulatorBusMix '%s' already contains bus '%s'. Only one representative channel for this bus added."),
 						*Mix.GetFullName(), *Channel.Bus->GetFullName());
 				}
-#endif // !UE_BUILD_SHIPPING
+#endif // UE_BUILD_SHIPPING
 				Channels.Emplace(BusMixId, FModulatorBusMixChannelProxy(Channel));
 			}
 			else
@@ -97,19 +97,29 @@ namespace AudioModulation
 		}
 	}
 
-	bool FModulatorBusMixProxy::CanDeactivate() const
+	bool FModulatorBusMixProxy::CanDestroy() const
+	{
+		if (Status != BusMixStatus::Stopped)
+		{
+			return false;
+		}
+
+		return !bAutoActivate || (bAutoActivate && SoundRefCount == 0);
+	}
+
+	bool FModulatorBusMixProxy::GetAutoActivate() const
 	{
-		return Status == BusMixStatus::Stopped;
+		return bAutoActivate;
 	}
 
-	void FModulatorBusMixProxy::SetActive()
+	void FModulatorBusMixProxy::SetEnabled()
 	{
-		Status = BusMixStatus::Active;
+		Status = BusMixStatus::Enabled;
 	}
 
 	void FModulatorBusMixProxy::SetStopping()
 	{
-		if (Status == BusMixStatus::Active)
+		if (Status == BusMixStatus::Enabled)
 		{
 			Status = BusMixStatus::Stopping;
 		}
@@ -156,4 +166,15 @@ namespace AudioModulation
 		return Name;
 	}
 #endif // !UE_BUILD_SHIPPING
+
+	int32 FModulatorBusMixProxy::DecRefSound()
+	{
+		check(SoundRefCount > 0);
+		return SoundRefCount--;
+	}
+
+	int32 FModulatorBusMixProxy::IncRefSound()
+	{
+		return SoundRefCount++;
+	}
 } // namespace AudioModulation
diff --git a/Engine/Plugins/Runtime/AudioModulation/Source/AudioModulation/Private/SoundModulatorBusMix.h b/Engine/Plugins/Runtime/AudioModulation/Source/AudioModulation/Private/SoundModulatorBusMix.h
--- a/Engine/Plugins/Runtime/AudioModulation/Source/AudioModulation/Private/SoundModulatorBusMix.h
+++ b/Engine/Plugins/Runtime/AudioModulation/Source/AudioModulation/Private/SoundModulatorBusMix.h
@@ -60,15 +60,16 @@ namespace AudioModulation
 	{
 		enum class BusMixStatus : uint8
 		{
-			Active,
+			Enabled,
 			Stopping,
 			Stopped
 		};
 
 		FModulatorBusMixProxy(const USoundModulatorBusMix& Mix);
 
-		bool CanDeactivate() const;
-		void SetActive();
+		bool CanDestroy() const;
+		bool GetAutoActivate() const;
+		void SetEnabled();
 		void SetStopping();
 		void Update(const float Elapsed, BusProxyMap& ProxyMap);
 
@@ -76,14 +77,19 @@ namespace AudioModulation
 		const FString& GetName() const;
 #endif // !UE_BUILD_SHIPPING
 
+		int32 IncRefSound();
+		int32 DecRefSound();
+
 		TMap<BusId, FModulatorBusMixChannelProxy> Channels;
 
 	private:
+		BusMixStatus Status;
+		int32 SoundRefCount;
+		uint8 bAutoActivate : 1;
+
 #if !UE_BUILD_SHIPPING
 		FString Name;
 #endif // !UE_BUILD_SHIPPING
-
-		BusMixStatus Status;
 	};
 	using BusMixProxyMap = TMap<BusMixId, FModulatorBusMixProxy>;
 } // namespace AudioModulation
diff --git a/Engine/Plugins/Runtime/AudioModulation/Source/AudioModulation/Private/SoundModulatorLFO.cpp b/Engine/Plugins/Runtime/AudioModulation/Source/AudioModulation/Private/SoundModulatorLFO.cpp
--- a/Engine/Plugins/Runtime/AudioModulation/Source/AudioModulation/Private/SoundModulatorLFO.cpp
+++ b/Engine/Plugins/Runtime/AudioModulation/Source/AudioModulation/Private/SoundModulatorLFO.cpp
@@ -16,7 +16,6 @@ USoundModulatorLFO::USoundModulatorLFO(const FObjectInitializer& ObjectInitializ
 	, Offset(0.5f)
 	, bLooping(1)
 	, bAutoActivate(0)
-	, bAutoDeactivate(0)
 {
 }
 
@@ -53,8 +52,8 @@ namespace AudioModulation
 		: Id(0)
 		, Offset(0.0f)
 		, Value(0.0f)
-		, bIsActive(false)
-		, bAutoDeactivate(false)
+		, bAutoActivate(false)
+		, SoundRefCount(0)
 	{
 	}
 
@@ -65,8 +64,8 @@ namespace AudioModulation
 #endif // !UE_BUILD_SHIPPING
 		, Offset(InLFO.Offset)
 		, Value(0.0f)
-		, bIsActive(false)
-		, bAutoDeactivate(InLFO.bAutoDeactivate)
+		, bAutoActivate(InLFO.bAutoActivate)
+		, SoundRefCount(0)
 	{
 		LFO.SetGain(InLFO.Amplitude);
 		LFO.SetFrequency(InLFO.Frequency);
@@ -78,19 +77,14 @@ namespace AudioModulation
 		LFO.Start();
 	}
 
-	bool FModulatorLFOProxy::CanDeactivate() const
+	float FModulatorLFOProxy::GetAmplitude() const
 	{
-		return !bIsActive && bAutoDeactivate;
+		return LFO.GetGain();
 	}
 
-	void FModulatorLFOProxy::ClearIsActive()
+	bool FModulatorLFOProxy::GetAutoActivate() const
 	{
-		bIsActive = false;
-	}
-
-	float FModulatorLFOProxy::GetAmplitude() const
-	{
-		return LFO.GetGain();
+		return bAutoActivate;
 	}
 
 	float FModulatorLFOProxy::GetFreq() const
@@ -125,11 +119,6 @@ namespace AudioModulation
 		LFO.SetFrequency(InFreq);
 	}
 
-	void FModulatorLFOProxy::SetIsActive()
-	{
-		bIsActive = true;
-	}
-
 	void FModulatorLFOProxy::Update(float InElapsed)
 	{
 		if (InElapsed > 0.0f && LFO.GetFrequency() > 0.0f)
@@ -140,4 +129,18 @@ namespace AudioModulation
 			Value = LFO.Generate() + Offset;
 		}
 	}
+
+	int32 FModulatorLFOProxy::DecRefSound()
+	{
+		check(SoundRefCount > 0);
+		SoundRefCount--;
+
+		return SoundRefCount;
+	}
+
+	int32 FModulatorLFOProxy::IncRefSound()
+	{
+		SoundRefCount++;
+		return SoundRefCount;
+	}
 } // namespace AudioModulation
diff --git a/Engine/Plugins/Runtime/AudioModulation/Source/AudioModulation/Private/SoundModulatorLFO.h b/Engine/Plugins/Runtime/AudioModulation/Source/AudioModulation/Private/SoundModulatorLFO.h
--- a/Engine/Plugins/Runtime/AudioModulation/Source/AudioModulation/Private/SoundModulatorLFO.h
+++ b/Engine/Plugins/Runtime/AudioModulation/Source/AudioModulation/Private/SoundModulatorLFO.h
@@ -53,13 +53,9 @@ public:
 	UPROPERTY(EditAnywhere, Category = Modulation, BlueprintReadWrite)
 	uint8 bLooping : 1;
 
-	/** Automatically activates LFO when sounds are referencing it */
+	/** Automatically activates/deactivates LFO when sounds referencing LFO asset are playing */
 	UPROPERTY(EditAnywhere, Category = Modulation, BlueprintReadWrite)
 	uint8 bAutoActivate : 1;
-
-	/** Automatically deactivates LFO when no sounds are referencing it */
-	UPROPERTY(EditAnywhere, Category = Modulation, BlueprintReadWrite)
-	uint8 bAutoDeactivate : 1;
 };
 
 namespace AudioModulation
@@ -70,9 +66,8 @@ namespace AudioModulation
 		FModulatorLFOProxy();
 		FModulatorLFOProxy(const USoundModulatorLFO& InLFO);
 
-		bool CanDeactivate() const;
-		void ClearIsActive();
 		float GetAmplitude() const;
+		bool GetAutoActivate() const;
 		float GetFreq() const;
 		LFOId GetId() const;
 #if !UE_BUILD_SHIPPING
@@ -81,9 +76,11 @@ namespace AudioModulation
 		float GetOffset() const;
 		float GetValue() const;
 		void SetFreq(float InFreq);
-		void SetIsActive();
 		void Update(float InElapsed);
 
+		int32 DecRefSound();
+		int32 IncRefSound();
+
 	private:
 		LFOId Id;
 
@@ -95,9 +92,8 @@ namespace AudioModulation
 		float Offset;
 		float Value;
 
-		uint8 bLimitedLifetime : 1;
-		uint8 bIsActive : 1;
-		uint8 bAutoDeactivate : 1;
+		uint8 bAutoActivate : 1;
+		int32 SoundRefCount;
 	};
 	using LFOProxyMap = TMap<BusId, FModulatorLFOProxy>;
 } // namespace AudioModulation
diff --git a/Engine/Plugins/Runtime/AudioModulation/Source/AudioModulation/Public/AudioModulation.h b/Engine/Plugins/Runtime/AudioModulation/Source/AudioModulation/Public/AudioModulation.h
--- a/Engine/Plugins/Runtime/AudioModulation/Source/AudioModulation/Public/AudioModulation.h
+++ b/Engine/Plugins/Runtime/AudioModulation/Source/AudioModulation/Public/AudioModulation.h
@@ -33,7 +33,8 @@ namespace AudioModulation
 		virtual void OnEditSource(const USoundModulationPluginSourceSettingsBase& Settings) override;
 #endif // WITH_EDITOR
 
-		virtual void OnInitSource(const uint32 SourceId, const FName& AudioComponentUserId, const uint32 NumChannels, USoundModulationPluginSourceSettingsBase* Settings) override;
+		virtual void OnInitSound(const ModulationSoundId SoundId, const USoundModulationPluginSourceSettingsBase& Settings) override;
+		virtual void OnInitSource(const uint32 SourceId, const FName& AudioComponentUserId, const uint32 NumChannels, const USoundModulationPluginSourceSettingsBase& Settings) override;
 
 #if !UE_BUILD_SHIPPING
 		virtual bool OnPostHelp(FCommonViewportClient* ViewportClient, const TCHAR* Stream) override;
@@ -41,6 +42,7 @@ namespace AudioModulation
 		virtual bool OnToggleStat(FCommonViewportClient* ViewportClient, const TCHAR* Stream) override;
 #endif // !UE_BUILD_SHIPPING
 
+		virtual void OnReleaseSound(const ModulationSoundId SoundId, const USoundModulationPluginSourceSettingsBase& Settings) override;
 		virtual void OnReleaseSource(const uint32 SourceId) override;
 		virtual void ProcessControls(const uint32 SourceId, FSoundModulationControls& Controls) override;
 		virtual void ProcessModulators(const float Elapsed) override;
diff --git a/Engine/Source/Runtime/AudioMixer/Private/AudioMixerSourceManager.cpp b/Engine/Source/Runtime/AudioMixer/Private/AudioMixerSourceManager.cpp
--- a/Engine/Source/Runtime/AudioMixer/Private/AudioMixerSourceManager.cpp
+++ b/Engine/Source/Runtime/AudioMixer/Private/AudioMixerSourceManager.cpp
@@ -613,7 +613,7 @@ namespace Audio
 		// Create the modulation plugin source effect
 		if (InitParams.ModulationPluginSettings != nullptr)
 		{
-			MixerDevice->ModulationInterface->OnInitSource(SourceId, InitParams.AudioComponentUserID, InitParams.NumInputChannels, InitParams.ModulationPluginSettings);
+			MixerDevice->ModulationInterface->OnInitSource(SourceId, InitParams.AudioComponentUserID, InitParams.NumInputChannels, *InitParams.ModulationPluginSettings);
 		}
 
 		AudioMixerThreadCommand([this, SourceId, InitParams]()
diff --git a/Engine/Source/Runtime/Engine/Classes/Sound/SoundBase.h b/Engine/Source/Runtime/Engine/Classes/Sound/SoundBase.h
--- a/Engine/Source/Runtime/Engine/Classes/Sound/SoundBase.h
+++ b/Engine/Source/Runtime/Engine/Classes/Sound/SoundBase.h
@@ -98,7 +98,7 @@ public:
 	TEnumAsByte<EMaxConcurrentResolutionRule::Type> MaxConcurrentResolutionRule_DEPRECATED;
 #endif // WITH_EDITORONLY_DATA
 
-	/** Number of times this cue is currently being played. */
+	/** Number of times this sound is currently being played (counted if sound is virtualized). */
 	int32 CurrentPlayCount;
 
 #if WITH_EDITORONLY_DATA
diff --git a/Engine/Source/Runtime/Engine/Private/ActiveSound.cpp b/Engine/Source/Runtime/Engine/Private/ActiveSound.cpp
--- a/Engine/Source/Runtime/Engine/Private/ActiveSound.cpp
+++ b/Engine/Source/Runtime/Engine/Private/ActiveSound.cpp
@@ -773,6 +773,17 @@ void FActiveSound::Stop(bool bStopNow)
 	if (Sound && !bIsStopping)
 	{
 		Sound->CurrentPlayCount = FMath::Max(Sound->CurrentPlayCount - 1, 0);
+		if (Sound->CurrentPlayCount == 0)
+		{
+			if (AudioDevice->ModulationInterface)
+			{
+				if (USoundModulationPluginSourceSettingsBase* ModulationSettings = FindModulationSettings())
+				{
+					const ModulationSoundId SoundId = static_cast<ModulationSoundId>(Sound->GetUniqueID());
+ 					AudioDevice->ModulationInterface->OnReleaseSound(SoundId, *ModulationSettings);
+				}
+			}
+		}
 	}
 
 	TArray<FWaveInstance*> ToDelete;
diff --git a/Engine/Source/Runtime/Engine/Private/AudioDevice.cpp b/Engine/Source/Runtime/Engine/Private/AudioDevice.cpp
--- a/Engine/Source/Runtime/Engine/Private/AudioDevice.cpp
+++ b/Engine/Source/Runtime/Engine/Private/AudioDevice.cpp
@@ -4330,7 +4330,6 @@ void FAudioDevice::AddNewActiveSoundInternal(const FActiveSound& NewActiveSound,
 		return;
 	}
 
-	check(ActiveSound->Sound);
 	check(ActiveSound->Sound == Sound);
 
 	if (GIsEditor)
@@ -4342,7 +4341,14 @@ void FAudioDevice::AddNewActiveSoundInternal(const FActiveSound& NewActiveSound,
 		}
 	}
 
-	++ActiveSound->Sound->CurrentPlayCount;
+	++Sound->CurrentPlayCount;
+	if (ModulationInterface && Sound->CurrentPlayCount == 1)
+	{
+		if (USoundModulationPluginSourceSettingsBase* ModulationSettings = ActiveSound->FindModulationSettings())
+		{
+			ModulationInterface->OnInitSound(static_cast<ModulationSoundId>(Sound->GetUniqueID()), *ModulationSettings);
+		}
+	}
 
 #if !(UE_BUILD_SHIPPING || UE_BUILD_TEST)
 	UE_LOG(LogAudio, VeryVerbose, TEXT("New ActiveSound %s Comp: %s Loc: %s"), *Sound->GetName(), *NewActiveSound.GetAudioComponentName(), *NewActiveSound.Transform.GetTranslation().ToString());
diff --git a/Engine/Source/Runtime/Engine/Public/IAudioExtensionPlugin.h b/Engine/Source/Runtime/Engine/Public/IAudioExtensionPlugin.h
--- a/Engine/Source/Runtime/Engine/Public/IAudioExtensionPlugin.h
+++ b/Engine/Source/Runtime/Engine/Public/IAudioExtensionPlugin.h
@@ -521,26 +521,35 @@ struct FSoundModulationControls
 	}
 };
 
+/** Identifier active sound types currently playing (or virtualized) */
+using ModulationSoundId = uint32;
+
 class IAudioModulation
 {
 public:
 	/** Virtual destructor */
 	virtual ~IAudioModulation() { }
 
-	/** Initialize the modulation plugin with the same rate and number of sources. */
+	/** Initialize the modulation plugin with the same rate and number of sources */
 	virtual void Initialize(const FAudioPluginInitializationParams& InitializationParams) { }
 
 #if WITH_EDITOR
-	/** Allows modulation plugin to respond to settings updates when client is editing settings. */
+	/** Allows modulation plugin to respond to settings updates when client is editing settings */
 	virtual void OnEditSource(const USoundModulationPluginSourceSettingsBase& Settings) { }
 #endif // WITH_EDITOR
 
-	/** Called when a source is assigned to a voice. */
-	virtual void OnInitSource(const uint32 SourceId, const FName& AudioComponentUserId, const uint32 NumChannels, USoundModulationPluginSourceSettingsBase* Settings) { }
+	/** Called when a USoundBase type begins playing a sound */
+	virtual void OnInitSound(const ModulationSoundId SoundId, const USoundModulationPluginSourceSettingsBase& Settings) { }
 
-	/** Called when a source is done playing and is released. */
+	/** Called when a source is assigned to a voice */
+	virtual void OnInitSource(const uint32 SourceId, const FName& AudioComponentUserId, const uint32 NumChannels, const USoundModulationPluginSourceSettingsBase& Settings) { }
+
+	/** Called when a source is done playing and is released */
 	virtual void OnReleaseSource(const uint32 SourceId) { }
 
+	/** Called when a USoundBase type stops playing any sounds */
+	virtual void OnReleaseSound(const ModulationSoundId SoundId, const USoundModulationPluginSourceSettingsBase& Settings) { }
+
 #if !UE_BUILD_SHIPPING
 	/** Request to post help from active plugin (non-shipping builds only) */
 	virtual bool OnPostHelp(FCommonViewportClient* ViewportClient, const TCHAR* Stream) { return false; };
