[NO NON-PII COMMIT MESSAGE]

diff --git a/Engine/Plugins/Runtime/AudioModulation/Source/AudioModulation/Private/SoundModulationTransform.cpp b/Engine/Plugins/Runtime/AudioModulation/Source/AudioModulation/Private/SoundModulationTransform.cpp
--- a/Engine/Plugins/Runtime/AudioModulation/Source/AudioModulation/Private/SoundModulationTransform.cpp
+++ b/Engine/Plugins/Runtime/AudioModulation/Source/AudioModulation/Private/SoundModulationTransform.cpp
@@ -3,7 +3,7 @@
 #include "SoundModulationTransform.h"
 
 #include "Audio.h"
-
+#include "DSP/Dsp.h"
 
 FSoundModulationInputTransform::FSoundModulationInputTransform()
 	: InputMin(0.0f)
