- Templatize activation and proxy state used for all modulator types (bus modulators (e.g. LFO), buses, and bus mixes)
- Add update modulator functionality so that any changes to uobject can be reflected in thread-safe proxy modulation types via blueprint/code
- Finish, fix & polish auto-activate modulation support
- Rename ModulatorBus types to ControlBus (to avoid confusion between USoundModulatorBusBase and USoundBusModulatorBase)

diff --git a/Engine/Plugins/Runtime/AudioModulation/Source/AudioModulation/Private/AudioModulation.cpp b/Engine/Plugins/Runtime/AudioModulation/Source/AudioModulation/Private/AudioModulation.cpp
--- a/Engine/Plugins/Runtime/AudioModulation/Source/AudioModulation/Private/AudioModulation.cpp
+++ b/Engine/Plugins/Runtime/AudioModulation/Source/AudioModulation/Private/AudioModulation.cpp
@@ -32,15 +32,15 @@ namespace AudioModulation
 	}
 
 #if WITH_EDITOR
-	void FAudioModulation::OnEditSource(const USoundModulationPluginSourceSettingsBase& Settings)
+	void FAudioModulation::OnEditPluginSettings(const USoundModulationPluginSourceSettingsBase& Settings)
 	{
-		Impl->OnEditSource(Settings);
+		Impl->OnEditPluginSettings(Settings);
 	}
 #endif // WITH_EDITOR
 
-	void FAudioModulation::OnInitSound(const ModulationSoundId SoundId, const USoundModulationPluginSourceSettingsBase& Settings)
+	void FAudioModulation::OnInitSound(ISoundModulatable& Sound, const USoundModulationPluginSourceSettingsBase& Settings)
 	{
-		Impl->OnInitSound(SoundId, Settings);
+		Impl->OnInitSound(Sound, Settings);
 	}
 
 	void FAudioModulation::OnInitSource(const uint32 SourceId, const FName& AudioComponentUserId, const uint32 NumChannels, const USoundModulationPluginSourceSettingsBase& Settings)
@@ -48,9 +48,9 @@ namespace AudioModulation
 		Impl->OnInitSource(SourceId, AudioComponentUserId, NumChannels, Settings);
 	}
 
-	void FAudioModulation::OnReleaseSound(const ModulationSoundId SoundId, const USoundModulationPluginSourceSettingsBase& Settings)
+	void FAudioModulation::OnReleaseSound(ISoundModulatable& Sound)
 	{
-		Impl->OnReleaseSound(SoundId, Settings);
+		Impl->OnReleaseSound(Sound);
 	}
 
 	void FAudioModulation::OnReleaseSource(const uint32 SourceId)
diff --git a/Engine/Plugins/Runtime/AudioModulation/Source/AudioModulation/Private/AudioModulationDebugger.cpp b/Engine/Plugins/Runtime/AudioModulation/Source/AudioModulation/Private/AudioModulationDebugger.cpp
--- a/Engine/Plugins/Runtime/AudioModulation/Source/AudioModulation/Private/AudioModulationDebugger.cpp
+++ b/Engine/Plugins/Runtime/AudioModulation/Source/AudioModulation/Private/AudioModulationDebugger.cpp
@@ -87,9 +87,6 @@ namespace AudioModulation
 				}
 
 				const FString ValueString  = FString::Printf(TEXT("%.6f"), LFOPair.GetValue());
-				const FString AmpString    = FString::Printf(TEXT("%.6f"), LFOPair.GetAmplitude());
-				const FString FreqString   = FString::Printf(TEXT("%.6f"), LFOPair.GetFreq());
-				const FString OffsetString = FString::Printf(TEXT("%.6f"), LFOPair.GetOffset());
 
 				Canvas.DrawShadowedString(RowX, Y, *Name, &Font, FColor::Green);
 				RowX += Width * CellWidth;
@@ -97,22 +94,13 @@ namespace AudioModulation
 				Canvas.DrawShadowedString(RowX, Y, *ValueString, &Font, FColor::Green);
 				RowX += Width * CellWidth;
 
-				Canvas.DrawShadowedString(RowX, Y, *AmpString, &Font, FColor::Green);
-				RowX += Width * CellWidth;
-
-				Canvas.DrawShadowedString(RowX, Y, *FreqString, &Font, FColor::Green);
-				RowX += Width * CellWidth;
-
-				Canvas.DrawShadowedString(RowX, Y, *OffsetString, &Font, FColor::Green);
-				RowX += Width * CellWidth;
-
 				Y += Height;
 			}
 
 			return Y;
 		}
 
-		int32 RenderStatMixMatrix(const TArray<FModulatorBusMixProxy>& FilteredMixes, const TArray<FModulatorBusProxy>& FilteredBuses, FCanvas& Canvas, int32 X, int32 Y, const UFont& Font)
+		int32 RenderStatMixMatrix(const TArray<FModulatorBusMixProxy>& FilteredMixes, const TArray<FControlBusProxy>& FilteredBuses, FCanvas& Canvas, int32 X, int32 Y, const UFont& Font)
 		{
 			int32 Height = 12;
 			int32 Width = 12;
@@ -139,14 +127,14 @@ namespace AudioModulation
 			{
 				CellWidth = FMath::Max(CellWidth, BusMix.GetName().Len());
 			}
-			for (const FModulatorBusProxy& BusPair : FilteredBuses)
+			for (const FControlBusProxy& BusPair : FilteredBuses)
 			{
 				CellWidth = FMath::Max(CellWidth, BusPair.GetName().Len());
 			}
 
 			// Draw Column Headers
 			int32 RowX = X;
-			for (const FModulatorBusProxy& BusPair : FilteredBuses)
+			for (const FControlBusProxy& BusPair : FilteredBuses)
 			{
 				FString Name = BusPair.GetName().Left(MaxNameLength);
 				if (Name.Len() < MaxNameLength)
@@ -180,13 +168,13 @@ namespace AudioModulation
 				ColumnY += Height; // Add before to leave space for column headers
 				RowX = X;
 
-				for (const FModulatorBusProxy& Bus : FilteredBuses)
+				for (const FControlBusProxy& Bus : FilteredBuses)
 				{
 					RowX += Width * CellWidth; // Add before to leave space for row headers
 
 					float Target = Bus.GetDefaultValue();
 					float Value = Bus.GetDefaultValue();
-					if (const FModulatorBusMixChannelProxy* ChannelProxy = BusMix.Channels.Find(Bus.GetBusId()))
+					if (const FModulatorBusMixChannelProxy* ChannelProxy = BusMix.Channels.Find(Bus.GetId()))
 					{
 						Target = ChannelProxy->Value.TargetValue;
 						Value  = ChannelProxy->Value.GetCurrentValue();
@@ -208,7 +196,7 @@ namespace AudioModulation
 			// Draw Sub-Totals & Totals
 			Canvas.DrawShadowedString(X, Y, *MixSubTotalHeader, &Font, FColor::Yellow);
 			RowX = X;
-			for (const FModulatorBusProxy& BusPair : FilteredBuses)
+			for (const FControlBusProxy& BusPair : FilteredBuses)
 			{
 				RowX += Width * CellWidth; // Add before to leave space for row headers
 
@@ -227,7 +215,7 @@ namespace AudioModulation
 
 			Canvas.DrawShadowedString(X, Y, *LFOSubTotalHeader, &Font, FColor::Yellow);
 			RowX = X;
-			for (const FModulatorBusProxy& BusPair : FilteredBuses)
+			for (const FControlBusProxy& BusPair : FilteredBuses)
 			{
 				RowX += Width * CellWidth; // Add before to leave space for row headers
 				const float Value = BusPair.GetLFOValue();
@@ -238,7 +226,7 @@ namespace AudioModulation
 			Y += Height;
 			Canvas.DrawShadowedString(X, Y, *TotalHeader, &Font, FColor::Yellow);
 			RowX = X;
-			for (const FModulatorBusProxy& BusPair : FilteredBuses)
+			for (const FControlBusProxy& BusPair : FilteredBuses)
 			{
 				RowX += Width * CellWidth; // Add before to leave space for row headers
 				const float Value = BusPair.GetValue();
@@ -299,8 +287,8 @@ namespace AudioModulation
 		}
 
 		static const int32 MaxFilteredBuses = 8;
-		TArray<FModulatorBusProxy> InFilteredBuses;
-		FilterDebugArray<BusProxyMap, FModulatorBusProxy>(ActiveBuses, BusStringFilter, MaxFilteredBuses, InFilteredBuses);
+		TArray<FControlBusProxy> InFilteredBuses;
+		FilterDebugArray<BusProxyMap, FControlBusProxy>(ActiveBuses, BusStringFilter, MaxFilteredBuses, InFilteredBuses);
 
 		static const int32 MaxFilteredMixes = 16;
 		TArray<FModulatorBusMixProxy> InFilteredMixes;
@@ -313,7 +301,7 @@ namespace AudioModulation
 		FAudioThread::RunCommandOnGameThread([this, InFilteredBuses, InFilteredMixes, InFilteredLFOs]()
 		{
 			FilteredBuses = InFilteredBuses;
-			FilteredBuses.Sort(&CompareNames<FModulatorBusProxy>);
+			FilteredBuses.Sort(&CompareNames<FControlBusProxy>);
 
 			FilteredMixes = InFilteredMixes;
 			FilteredMixes.Sort(&CompareNames<FModulatorBusMixProxy>);
diff --git a/Engine/Plugins/Runtime/AudioModulation/Source/AudioModulation/Private/AudioModulationDebugger.h b/Engine/Plugins/Runtime/AudioModulation/Source/AudioModulation/Private/AudioModulationDebugger.h
--- a/Engine/Plugins/Runtime/AudioModulation/Source/AudioModulation/Private/AudioModulationDebugger.h
+++ b/Engine/Plugins/Runtime/AudioModulation/Source/AudioModulation/Private/AudioModulationDebugger.h
@@ -5,8 +5,8 @@
 
 #if WITH_AUDIOMODULATION
 #if !UE_BUILD_SHIPPING
-#include "SoundModulatorBus.h"
-#include "SoundModulatorBusMix.h"
+#include "SoundControlBus.h"
+#include "SoundControlBusMix.h"
 #include "SoundModulatorLFO.h"
 
 
@@ -30,7 +30,7 @@ namespace AudioModulation
 		uint8 bShowRenderStatLFO : 1;
 		uint8 bShowRenderStatMix : 1;
 
-		TArray<FModulatorBusProxy>    FilteredBuses;
+		TArray<FControlBusProxy>    FilteredBuses;
 		TArray<FModulatorLFOProxy>    FilteredLFOs;
 		TArray<FModulatorBusMixProxy> FilteredMixes;
 
diff --git a/Engine/Plugins/Runtime/AudioModulation/Source/AudioModulation/Private/AudioModulationImpl.cpp b/Engine/Plugins/Runtime/AudioModulation/Source/AudioModulation/Private/AudioModulationImpl.cpp
--- a/Engine/Plugins/Runtime/AudioModulation/Source/AudioModulation/Private/AudioModulationImpl.cpp
+++ b/Engine/Plugins/Runtime/AudioModulation/Source/AudioModulation/Private/AudioModulationImpl.cpp
@@ -22,6 +22,40 @@ DECLARE_DWORD_COUNTER_STAT(TEXT("Mix Count"), STAT_AudioModulationMixCount, STAT
 
 namespace
 {
+#if WITH_EDITOR
+	// Checks whether new inputs being modified require stopping referencing preview sound for proxies to be updated and take change.
+	template <typename T>
+	bool InputUpdateRequiresStop(const TArray<T>& NewInputs, const TArray<AudioModulation::FModulationInputProxy>& CurrentInputProxies)
+	{
+		check(IsInAudioThread());
+
+		if (NewInputs.Num() != CurrentInputProxies.Num())
+		{
+			return true;
+		}
+
+		for (int32 i = 0; i < NewInputs.Num(); ++i)
+		{
+			if (const USoundControlBusBase* NewBus = NewInputs[i].GetBus())
+			{
+				if (NewBus->GetUniqueID() != CurrentInputProxies[i].BusId)
+				{
+					return true;
+				}
+			}
+			else
+			{
+				if (CurrentInputProxies[i].BusId != INDEX_NONE)
+				{
+					return true;
+				}
+			}
+		}
+
+		return false;
+	};
+#endif
+
 	void MixInModulationValue(ESoundModulatorOperator& Operator, float ModStageValue, float& Value)
 	{
 		switch (Operator)
@@ -52,114 +86,202 @@ namespace
 
 namespace AudioModulation
 {
+	FAudioModulationImpl::FAudioModulationImpl()
+		: PreviewSound(nullptr)
+	{
+	}
+
 	void FAudioModulationImpl::Initialize(const FAudioPluginInitializationParams& InitializationParams)
 	{
 		SourceSettings.AddDefaulted(InitializationParams.NumSources);
 	}
 
 #if WITH_EDITOR
-	void FAudioModulationImpl::OnEditSource(const USoundModulationPluginSourceSettingsBase& Settings)
+	void FAudioModulationImpl::OnEditPluginSettings(const USoundModulationPluginSourceSettingsBase& InSettings)
 	{
-		FModulationSettingsProxy NewProxy(*CastChecked<USoundModulationSettings>(&Settings));
-		const uint32 ObjectID = Settings.GetUniqueID();
-
-		auto EditSource = [this, NewProxy, ObjectID]()
+		// Find if sound is being referenced and auditioned and stop immediately.
+		// This informs user that modifying sound's settings does not translate to
+		// currently playing sound.
+		const uint32 SettingsId = InSettings.GetUniqueID();
+		auto UpdateSettings = [this, SettingsId]()
 		{
-			for (FModulationSettingsProxy& CurrentSettingsProxy : SourceSettings)
+			if (PreviewSound)
 			{
-				if (CurrentSettingsProxy.ObjectID == ObjectID)
+				USoundModulationPluginSourceSettingsBase* SettingsBase = PreviewSound->FindModulationSettings();
+				if (SettingsBase && SettingsId == SettingsBase->GetUniqueID())
 				{
-					CurrentSettingsProxy = NewProxy;
+					const USoundModulationSettings* Settings = CastChecked<USoundModulationSettings>(SettingsBase);
+
+					bool bShouldStop = false;
+					if (InputUpdateRequiresStop<FSoundVolumeModulationInput>(Settings->Volume.Inputs, PreviewSettings.Volume.InputProxies))
+					{
+						bShouldStop = true;
+					}
+					else if (InputUpdateRequiresStop<FSoundHPFModulationInput>(Settings->Highpass.Inputs, PreviewSettings.Highpass.InputProxies))
+					{
+						bShouldStop = true;
+					}
+					else if (InputUpdateRequiresStop<FSoundLPFModulationInput>(Settings->Lowpass.Inputs, PreviewSettings.Lowpass.InputProxies))
+					{
+						bShouldStop = true;
+					}
+					else if (InputUpdateRequiresStop<FSoundPitchModulationInput>(Settings->Pitch.Inputs, PreviewSettings.Pitch.InputProxies))
+					{
+						bShouldStop = true;
+					}
+
+					if (bShouldStop)
+					{
+						PreviewSound->Stop();
+						PreviewSound = nullptr;
+						PreviewSettings = FModulationSettingsProxy();
+					}
+					else
+					{
+						PreviewSettings = FModulationSettingsProxy(*Settings);
+						for (FModulationSettingsProxy& SourceSetting : SourceSettings)
+						{
+							if (SourceSetting.GetId() == PreviewSettings.GetId())
+							{
+								SourceSetting = PreviewSettings;
+							}
+						}
+
+						for (TPair<uint32, FModulationSettingsProxy>& Pair : SoundSettings)
+						{
+							if (Pair.Value.GetId() == PreviewSettings.GetId())
+							{
+								Pair.Value = PreviewSettings;
+							}
+						}
+					}
 				}
 			}
 		};
 
-		IsInAudioThread() ? EditSource() : FAudioThread::RunCommandOnAudioThread(EditSource);
+		IsInAudioThread() ? UpdateSettings() : FAudioThread::RunCommandOnAudioThread(UpdateSettings);
 	}
 #endif // WITH_EDITOR
 
-	void FAudioModulationImpl::OnInitSound(const ModulationSoundId SoundId, const USoundModulationPluginSourceSettingsBase& InSettings)
+	void FAudioModulationImpl::OnInitSound(ISoundModulatable& InSound, const USoundModulationPluginSourceSettingsBase& InSettings)
 	{
 		check(IsInAudioThread());
 
-		if (const USoundModulationSettings* Settings = CastChecked<USoundModulationSettings>(&InSettings))
+		const USoundModulationSettings* Settings = CastChecked<USoundModulationSettings>(&InSettings);
+		const uint32 SoundId = InSound.GetObjectId();
+
+		if (InSound.IsPreviewSound())
+		{
+			PreviewSound = &InSound;
+			PreviewSettings = FModulationSettingsProxy(*Settings);
+		}
+
+		if (!SoundSettings.Contains(SoundId))
 		{
-			for (const USoundModulatorBusMix* Mix : Settings->Mixes)
+			SoundSettings.Add(SoundId, FModulationSettingsProxy(*Settings));
+		}
+
+		for (const USoundControlBusMix* Mix : Settings->Mixes)
+		{
+			if (Mix && (Mix->bAutoActivate || InSound.IsPreviewSound()))
 			{
-				if (Mix)
-				{
-					ActivateBusMix(*Mix, /* bReset */ false, SoundId);
-				}
+				ActivateBusMix(*Mix, &InSound);
 			}
+		}
 
-			auto CheckRefActivate = [this](const USoundModulatorBusBase* Bus)
+		auto CheckRefActivate = [this, &InSound](const USoundControlBusBase* Bus)
+		{
+			if (!Bus)
 			{
-				if (!Bus)
-				{
-					return;
-				}
+				return;
+			}
 
-				if (Bus->bAutoActivate)
-				{
-					ActivateBus(*Bus);
-				}
+			if (Bus->bAutoActivate || InSound.IsPreviewSound())
+			{
+				ActivateBus(*Bus, &InSound);
+			}
 
-				for (const USoundModulatorBase* Modulator : Bus->Modulators)
+			for (const USoundBusModulatorBase* Modulator : Bus->Modulators)
+			{
+				if (const USoundBusModulatorLFO* LFO = Cast<USoundBusModulatorLFO>(Modulator))
 				{
-					if (const USoundModulatorLFO* LFO = Cast<USoundModulatorLFO>(Modulator))
+					if (LFO->bAutoActivate || InSound.IsPreviewSound())
 					{
-						if (LFO->bAutoActivate)
-						{
-							FModulatorLFOProxy LFOProxy(*LFO);
-							ActivateLFO(*LFO);
-						}
+						FModulatorLFOProxy LFOProxy(*LFO);
+						ActivateLFO(*LFO, &InSound);
 					}
 				}
-			};
-
-			for (const FSoundVolumeModulationInput& Input : Settings->Volume.Inputs)
-			{
-				CheckRefActivate(Cast<USoundModulatorBusBase>(Input.Bus));
 			}
+		};
 
-			for (const FSoundPitchModulationInput& Input : Settings->Pitch.Inputs)
-			{
-				CheckRefActivate(Cast<USoundModulatorBusBase>(Input.Bus));
-			}
+		for (const FSoundVolumeModulationInput& Input : Settings->Volume.Inputs)
+		{
+			CheckRefActivate(Cast<USoundControlBusBase>(Input.Bus));
+		}
 
-			for (const FSoundLPFModulationInput& Input : Settings->Lowpass.Inputs)
-			{
-				CheckRefActivate(Cast<USoundModulatorBusBase>(Input.Bus));
-			}
+		for (const FSoundPitchModulationInput& Input : Settings->Pitch.Inputs)
+		{
+			CheckRefActivate(Cast<USoundControlBusBase>(Input.Bus));
+		}
 
-			for (const FSoundHPFModulationInput& Input : Settings->Highpass.Inputs)
-			{
-				CheckRefActivate(Cast<USoundModulatorBusBase>(Input.Bus));
-			}
+		for (const FSoundLPFModulationInput& Input : Settings->Lowpass.Inputs)
+		{
+			CheckRefActivate(Cast<USoundControlBusBase>(Input.Bus));
+		}
+
+		for (const FSoundHPFModulationInput& Input : Settings->Highpass.Inputs)
+		{
+			CheckRefActivate(Cast<USoundControlBusBase>(Input.Bus));
 		}
 	}
 
-	void FAudioModulationImpl::OnInitSource(const uint32 SourceId, const FName& AudioComponentUserId, const uint32 NumChannels, const USoundModulationPluginSourceSettingsBase& InSettings)
+	void FAudioModulationImpl::OnInitSource(const uint32 InSourceId, const FName& AudioComponentUserId, const uint32 NumChannels, const USoundModulationPluginSourceSettingsBase& InSettings)
 	{
 		check(IsInAudioThread());
 
 		if (const USoundModulationSettings* Settings = CastChecked<USoundModulationSettings>(&InSettings))
 		{
-			SourceSettings[SourceId] = FModulationSettingsProxy(*Settings);
+			SourceSettings[InSourceId] = FModulationSettingsProxy(*Settings);
 		}
 	}
 
-	void FAudioModulationImpl::OnReleaseSource(const uint32 SourceId)
+	void FAudioModulationImpl::OnReleaseSource(const uint32 InSourceId)
 	{
 		check(IsInAudioThread());
 
-		SourceSettings[SourceId] = FModulationSettingsProxy();
+		SourceSettings[InSourceId] = FModulationSettingsProxy();
 	}
 
-	void FAudioModulationImpl::OnReleaseSound(const ModulationSoundId SoundId, const USoundModulationPluginSourceSettingsBase& Settings)
+	void FAudioModulationImpl::OnReleaseSound(ISoundModulatable& InSound)
 	{
 		check(IsInAudioThread());
+		check(InSound.GetObjectId() != INDEX_NONE);
+
+		const FModulationSettingsProxy& Settings = SoundSettings.FindChecked(InSound.GetObjectId());
+		for (FBusMixId MixId : Settings.Mixes)
+		{
+			DeactivateBusMix(MixId, &InSound);
+		}
 
+		for (const FModulationInputProxy& Input : Settings.Volume.InputProxies)
+		{
+			DeactivateBus(Input.BusId, &InSound);
+		}
+
+		for (const FModulationInputProxy& Input : Settings.Pitch.InputProxies)
+		{
+			DeactivateBus(Input.BusId, &InSound);
+		}
+
+		for (const FModulationInputProxy& Input : Settings.Lowpass.InputProxies)
+		{
+			DeactivateBus(Input.BusId, &InSound);
+		}
+
+		for (const FModulationInputProxy& Input : Settings.Highpass.InputProxies)
+		{
+			DeactivateBus(Input.BusId, &InSound);
+		}
 	}
 
 #if !UE_BUILD_SHIPPING
@@ -179,144 +301,171 @@ namespace AudioModulation
 	}
 #endif // !UE_BUILD_SHIPPING
 
-	void FAudioModulationImpl::ActivateBus(const USoundModulatorBusBase& Bus, const ModulationSoundId SoundId)
+	void FAudioModulationImpl::ActivateBus(const USoundControlBusBase& InBus, const ISoundModulatable* InSound)
 	{
-		const FModulatorBusProxy NewBusProxy(Bus);
+		const FControlBusProxy NewBusProxy(InBus);
+		const bool bCanCreateNew = InBus.CanAutoActivate(InSound);
 
-		auto ActivateBusInternal = [this, NewBusProxy, SoundId]()
+		auto ActivateBusInternal = [this, NewBusProxy, InSound, bCanCreateNew]()
 		{
-			const BusId NewBusId = NewBusProxy.GetBusId();
-			FModulatorBusProxy* BusProxy = ActiveBuses.Find(NewBusId);
-			if (!BusProxy)
+			const FBusId NewBusId = NewBusProxy.GetId();
+			FControlBusProxy* BusProxy = ActiveBuses.Find(NewBusId);
+			if (!BusProxy && bCanCreateNew)
 			{
 				BusProxy = &ActiveBuses.Add(NewBusId, NewBusProxy);
 			}
 
-			if (BusProxy->GetAutoActivate() && SoundId != INDEX_NONE)
+			if (BusProxy && InSound)
 			{
-				BusProxy->IncRefSound();
+				BusProxy->OnInitSound(*InSound);
 			}
 		};
 
 		IsInAudioThread() ? ActivateBusInternal() : FAudioThread::RunCommandOnAudioThread(ActivateBusInternal);
 	}
 
-	void FAudioModulationImpl::ActivateBusMix(const USoundModulatorBusMix& BusMix, bool bReset, const ModulationSoundId SoundId)
+	void FAudioModulationImpl::ActivateBusMix(const USoundControlBusMix& InBusMix, const ISoundModulatable* InSound)
 	{
-		const BusMixId MixId = static_cast<BusMixId>(BusMix.GetUniqueID());
-		FModulatorBusMixProxy MixProxy(BusMix);
+		const FModulatorBusMixProxy NewMixProxy(InBusMix);
+		const bool bCanCreateNew = InBusMix.CanAutoActivate(InSound);
 
-		auto ActivateMixInternal = [this, MixId, MixProxy, bReset, SoundId]()
+		auto ActivateMixInternal = [this, NewMixProxy, InSound, bCanCreateNew]()
 		{
-			FModulatorBusMixProxy* BusMixProxy = nullptr;
-			if (FModulatorBusMixProxy* ExistingMixProxy = ActiveBusMixes.Find(MixId))
+			const FBusMixId MixId = static_cast<FBusMixId>(NewMixProxy.GetId());
+			FModulatorBusMixProxy* BusMixProxy = ActiveBusMixes.Find(MixId);
+			if (BusMixProxy)
 			{
 				// Enable in case mix is currently stopping but not yet stopped.
-				ExistingMixProxy->SetEnabled();
-				if (bReset)
-				{
-					for (int32 i = 0; i < ExistingMixProxy->Channels.Num(); ++i)
-					{
-						if (ExistingMixProxy->Channels[i].BusId == MixProxy.Channels[i].BusId)
-						{
-							ExistingMixProxy->Channels[i].Value.TargetValue = MixProxy.Channels[i].Value.TargetValue;
-						}
-					}
-				}
-
-				BusMixProxy = ExistingMixProxy;
+				BusMixProxy->SetEnabled();
 			}
-			else
+			else if (bCanCreateNew)
 			{
-				BusMixProxy = &ActiveBusMixes.Emplace(MixId, MixProxy);
+				BusMixProxy = &ActiveBusMixes.Add(MixId, NewMixProxy);
 			}
 
-			if (BusMixProxy->GetAutoActivate() && SoundId != INDEX_NONE)
+			if (BusMixProxy && InSound)
 			{
-				BusMixProxy->IncRefSound();
+				BusMixProxy->OnInitSound(*InSound);
 			}
 		};
 
 		IsInAudioThread() ? ActivateMixInternal() : FAudioThread::RunCommandOnAudioThread(ActivateMixInternal);
 	}
 
-	void FAudioModulationImpl::ActivateLFO(const USoundModulatorLFO& LFO, const ModulationSoundId SoundId)
+	void FAudioModulationImpl::ActivateLFO(const USoundBusModulatorLFO& InLFO, const ISoundModulatable* InSound)
 	{
-		const AudioModulation::LFOId LFOId = static_cast<const AudioModulation::LFOId>(LFO.GetUniqueID());
-		const FModulatorLFOProxy NewLFOProxy(LFO);
+		const FModulatorLFOProxy NewLFOProxy(InLFO);
+		const bool bCanCreateNew = InLFO.CanAutoActivate(InSound);
 
-		auto ActivateLFOInternal = [this, LFOId, NewLFOProxy, SoundId]()
+		auto ActivateLFOInternal = [this, NewLFOProxy, InSound, bCanCreateNew]()
 		{
-			FModulatorLFOProxy* LFOProxy = nullptr;
-			if (!LFOProxy)
+			const AudioModulation::FLFOId LFOId = NewLFOProxy.GetId();
+			FModulatorLFOProxy* LFOProxy = ActiveLFOs.Find(LFOId);
+			if (!LFOProxy && bCanCreateNew)
 			{
 				LFOProxy = &ActiveLFOs.Add(LFOId, NewLFOProxy);
 			}
 
-			if (SoundId != INDEX_NONE)
+			if (LFOProxy && InSound)
 			{
-				LFOProxy->IncRefSound();
+				LFOProxy->OnInitSound(*InSound);
 			}
 		};
 
 		IsInAudioThread() ? ActivateLFOInternal() : FAudioThread::RunCommandOnAudioThread(ActivateLFOInternal);
 	}
 
-	float FAudioModulationImpl::CalculateModulationValue(FModulationPatchProxy& Proxy) const
+	float FAudioModulationImpl::CalculateModulationValue(FModulationPatchProxy& OutProxy) const
 	{
-		float Value = Proxy.DefaultInputValue;
+		float OutValue = OutProxy.DefaultInputValue;
 
-		if (!Proxy.OutputProxy.bInitialized)
+		float& OutSampleHold = OutProxy.OutputProxy.SampleAndHoldValue;
+		if (!OutProxy.OutputProxy.bInitialized)
 		{
-			Proxy.OutputProxy.SampleAndHoldValue = Proxy.DefaultInputValue;
+			switch (OutProxy.OutputProxy.Operator)
+			{
+				case ESoundModulatorOperator::Max:
+				{
+					OutSampleHold = OutProxy.OutputProxy.Transform.OutputMin;
+				}
+				break;
+
+				case ESoundModulatorOperator::Min:
+				{
+					OutSampleHold = OutProxy.OutputProxy.Transform.OutputMax;
+				}
+				break;
+
+				case ESoundModulatorOperator::Multiply:
+				{
+					OutSampleHold = 1.0f;
+				}
+				break;
+
+				default:
+				{
+					static_assert(static_cast<int32>(ESoundModulatorOperator::Count) == 3, "Possible missing operator switch case coverage");
+				}
+				break;
+			}
 		}
 
-		for (const FModulationInputProxy& InputProxy : Proxy.InputProxies)
+		for (const FModulationInputProxy& InputProxy : OutProxy.InputProxies)
 		{
 			if (InputProxy.bSampleAndHold)
 			{
-				if (!Proxy.OutputProxy.bInitialized)
+				if (!OutProxy.OutputProxy.bInitialized)
 				{
-					if (const FModulatorBusProxy* BusProxy = ActiveBuses.Find(InputProxy.BusId))
+					if (const FControlBusProxy* BusProxy = ActiveBuses.Find(InputProxy.BusId))
 					{
 						float ModStageValue = BusProxy->GetValue();
 						InputProxy.Transform.Apply(ModStageValue);
-						MixInModulationValue(Proxy.OutputProxy.Operator, ModStageValue, Proxy.OutputProxy.SampleAndHoldValue);
+						MixInModulationValue(OutProxy.OutputProxy.Operator, ModStageValue, OutSampleHold);
 					}
 				}
 			}
 			else
 			{
-				if (const FModulatorBusProxy* BusProxy = ActiveBuses.Find(InputProxy.BusId))
+				if (const FControlBusProxy* BusProxy = ActiveBuses.Find(InputProxy.BusId))
 				{
 					float ModStageValue = BusProxy->GetValue();
 					InputProxy.Transform.Apply(ModStageValue);
-					MixInModulationValue(Proxy.OutputProxy.Operator, ModStageValue, Value);
+					MixInModulationValue(OutProxy.OutputProxy.Operator, ModStageValue, OutValue);
 				}
 			}
 		}
 
-		Proxy.OutputProxy.bInitialized = true;
-		MixInModulationValue(Proxy.OutputProxy.Operator, Proxy.OutputProxy.SampleAndHoldValue, Value);
-		Proxy.OutputProxy.Transform.Apply(Value);
-		return Value;
+		if (!OutProxy.OutputProxy.bInitialized)
+		{
+			const float OutputMin = OutProxy.OutputProxy.Transform.OutputMin;
+			const float OutputMax = OutProxy.OutputProxy.Transform.OutputMax;
+			OutSampleHold = FMath::Clamp(OutSampleHold, OutputMin, OutputMax);
+			OutProxy.OutputProxy.bInitialized = true;
+		}
+
+		OutProxy.OutputProxy.Transform.Apply(OutValue);
+		MixInModulationValue(OutProxy.OutputProxy.Operator, OutSampleHold, OutValue);
+		return OutValue;
 	}
 
-	void FAudioModulationImpl::DeactivateBusMix(const BusMixId BusMixId, const ModulationSoundId SoundId)
+	void FAudioModulationImpl::DeactivateBusMix(const FBusMixId InBusMixId, const ISoundModulatable* InSound)
 	{
-		auto DeactivateMix = [this, BusMixId, SoundId]()
+		auto DeactivateMix = [this, InBusMixId, InSound]()
 		{
 			check(IsInAudioThread());
 
-			if (FModulatorBusMixProxy* Mix = ActiveBusMixes.Find(BusMixId))
+			if (FModulatorBusMixProxy* Mix = ActiveBusMixes.Find(InBusMixId))
 			{
-				if (SoundId == INDEX_NONE)
+				if (!InSound)
 				{
-					Mix->SetStopping();
+					if (!Mix->GetAutoActivate())
+					{
+						Mix->SetStopping();
+					}
 				}
-				else if (Mix->GetAutoActivate())
+				else if (Mix->OnReleaseSound(*InSound) == 0)
 				{
-					if (Mix->DecRefSound() == 0)
+					if (Mix->GetAutoActivate())
 					{
 						Mix->SetStopping();
 					}
@@ -327,23 +476,36 @@ namespace AudioModulation
 		IsInAudioThread() ? DeactivateMix() : FAudioThread::RunCommandOnAudioThread(DeactivateMix);
 	}
 
-	void FAudioModulationImpl::DeactivateBus(const BusId BusId, const ModulationSoundId SoundId)
+	void FAudioModulationImpl::DeactivateBus(const FBusId InBusId, const ISoundModulatable* InSound)
 	{
-		auto Deactivate = [this, BusId, SoundId]()
+		auto Deactivate = [this, InBusId, InSound]()
 		{
 			check(IsInAudioThread());
 
-			if (FModulatorBusProxy* Bus = ActiveBuses.Find(BusId))
+			if (FControlBusProxy* Bus = ActiveBuses.Find(InBusId))
 			{
-				if (SoundId == INDEX_NONE)
+				if (!InSound)
+				{
+					if (!Bus->GetAutoActivate())
+					{
+						ActiveBuses.Remove(InBusId);
+					}
+				}
+				else if (Bus->OnReleaseSound(*InSound) == 0)
 				{
-					ActiveBuses.Remove(BusId);
+					if (Bus->GetAutoActivate())
+					{
+						ActiveBuses.Remove(InBusId);
+					}
 				}
-				else if (Bus->GetAutoActivate())
+
+				// Only pass along to referenced LFOs if deactivating
+				// via notification of sound release
+				if (InSound)
 				{
-					if (Bus->DecRefSound() == 0)
+					for (FLFOId LFOId : Bus->GetLFOIds())
 					{
-						ActiveBuses.Remove(BusId);
+						DeactivateLFO(LFOId, InSound);
 					}
 				}
 			}
@@ -352,26 +514,46 @@ namespace AudioModulation
 		IsInAudioThread() ? Deactivate() : FAudioThread::RunCommandOnAudioThread(Deactivate);
 	}
 
-	void FAudioModulationImpl::DeactivateLFO(const LFOId LFOId, const ModulationSoundId SoundId)
+	void FAudioModulationImpl::DeactivateLFO(const FLFOId InLFOId, const ISoundModulatable* InSound)
 	{
-		if (!IsInAudioThread())
+		auto Deactivate = [this, InLFOId, InSound]()
 		{
-			FAudioThread::RunCommandOnAudioThread([this, LFOId, SoundId]()
+			check(IsInAudioThread());
+
+			if (FModulatorLFOProxy* LFO = ActiveLFOs.Find(InLFOId))
 			{
-				DeactivateLFO(LFOId);
-			});
+				if (!InSound)
+				{
+					if (!LFO->GetAutoActivate())
+					{
+						ActiveBuses.Remove(InLFOId);
+					}
+				}
+				else if (LFO->OnReleaseSound(*InSound) == 0)
+				{
+					if (LFO->GetAutoActivate())
+					{
+						ActiveBuses.Remove(InLFOId);
+					}
+				}
+			}
+		};
 
-			return;
-		}
+		IsInAudioThread() ? Deactivate() : FAudioThread::RunCommandOnAudioThread(Deactivate);
+	}
+
+	bool FAudioModulationImpl::IsBusActive(const FBusId InBusId) const
+	{
+		check(IsInAudioThread());
 
-		ActiveLFOs.Remove(LFOId);
+		return ActiveBuses.Contains(InBusId);
 	}
 
-	bool FAudioModulationImpl::IsBusActive(const BusId BusId) const
+	bool FAudioModulationImpl::IsLFOActive(const FLFOId InLFOId) const
 	{
 		check(IsInAudioThread());
 
-		return ActiveBuses.Contains(BusId);
+		return ActiveLFOs.Contains(InLFOId);
 	}
 
 	void FAudioModulationImpl::ProcessControls(const uint32 SourceId, FSoundModulationControls& Controls)
@@ -391,13 +573,13 @@ namespace AudioModulation
 		check(IsInAudioThread());
 
 		// Update LFOs (prior to bus mixing to avoid single-frame latency)
-		for (TPair<LFOId, FModulatorLFOProxy>& Pair : ActiveLFOs)
+		for (TPair<FLFOId, FModulatorLFOProxy>& Pair : ActiveLFOs)
 		{
 			Pair.Value.Update(Elapsed);
 		}
 
 		// Reset buses & refresh cached LFO
-		for (TPair<BusId, FModulatorBusProxy>& Pair : ActiveBuses)
+		for (TPair<FBusId, FControlBusProxy>& Pair : ActiveBuses)
 		{
 			Pair.Value.Reset();
 			Pair.Value.MixLFO(ActiveLFOs);
@@ -405,8 +587,8 @@ namespace AudioModulation
 
 		// Update mix values and apply to prescribed buses.
 		// Track bus mixes ready to remove
-		TArray<BusMixId> MixesToDeactivate;
-		for (TPair<BusMixId, FModulatorBusMixProxy>& Pair : ActiveBusMixes)
+		TArray<FBusMixId> MixesToDeactivate;
+		for (TPair<FBusMixId, FModulatorBusMixProxy>& Pair : ActiveBusMixes)
 		{
 			Pair.Value.Update(Elapsed, ActiveBuses);
 			if (Pair.Value.CanDestroy())
@@ -417,7 +599,7 @@ namespace AudioModulation
 		}
 
 		// Destroy mixes that have stopped
-		for (const BusMixId& MixId : MixesToDeactivate)
+		for (const FBusMixId& MixId : MixesToDeactivate)
 		{
 			ActiveBusMixes.Remove(MixId);
 		}
@@ -431,56 +613,31 @@ namespace AudioModulation
 #endif // !UE_BUILD_SHIPPING
 	}
 
-	void FAudioModulationImpl::SetBusDefault(const USoundModulatorBusBase& Bus, const float Value)
+	void FAudioModulationImpl::UpdateModulator(const USoundModulatorBase& InModulator)
 	{
-		auto BusId = static_cast<const AudioModulation::BusId>(Bus.GetUniqueID());
-		if (FModulatorBusProxy* BusProxy = ActiveBuses.Find(BusId))
+		if (const USoundBusModulatorBase* BusModulator = Cast<USoundBusModulatorBase>(&InModulator))
 		{
-			BusProxy->SetDefaultValue(Value);
-		}
-	}
-
-	void FAudioModulationImpl::SetBusRange(const USoundModulatorBusBase& Bus, const FVector2D& Range)
-	{
-		auto BusId = static_cast<const AudioModulation::BusId>(Bus.GetUniqueID());
-		if (FModulatorBusProxy* BusProxy = ActiveBuses.Find(BusId))
-		{
-			BusProxy->SetRange(Range);
+			if (FModulatorLFOProxy* LFOProxy = ActiveLFOs.Find(static_cast<FLFOId>(BusModulator->GetUniqueID())))
+			{
+				LFOProxy->OnUpdateProxy(InModulator);
+			}
 		}
-	}
 
-	void FAudioModulationImpl::SetBusMixChannel(const USoundModulatorBusMix& BusMix, const USoundModulatorBusBase& Bus, const float TargetValue)
-	{
-		const AudioModulation::BusMixId MixId = static_cast<const AudioModulation::BusMixId>(BusMix.GetUniqueID());
-		const AudioModulation::BusId BusId = static_cast<const AudioModulation::BusId>(Bus.GetUniqueID());
-
-		auto SetChannel = [this, MixId, BusId, TargetValue]()
+		if (const USoundControlBusBase* Bus = Cast<USoundControlBusBase>(&InModulator))
 		{
-			if (FModulatorBusMixProxy* BusMixProxy = ActiveBusMixes.Find(MixId))
+			if (FControlBusProxy* BusProxy = ActiveBuses.Find(static_cast<FBusId>(Bus->GetUniqueID())))
 			{
-				if (FModulatorBusMixChannelProxy* ChannelProxy = BusMixProxy->Channels.Find(BusId))
-				{
-					ChannelProxy->Value.TargetValue = TargetValue;
-				}
+				BusProxy->OnUpdateProxy(InModulator);
 			}
-		};
-
-		IsInAudioThread() ? SetChannel() : FAudioThread::RunCommandOnAudioThread(SetChannel);
-	}
-
-	void FAudioModulationImpl::SetLFOFrequency(const USoundModulatorLFO& LFO, const float Freq)
-	{
-		const AudioModulation::LFOId LFOId = static_cast<const AudioModulation::LFOId>(LFO.GetUniqueID());
+		}
 
-		auto SetFreq = [this, LFOId, Freq]()
+		if (const USoundControlBusMix* Mix = Cast<USoundControlBusMix>(&InModulator))
 		{
-			if (FModulatorLFOProxy* LFOProxy = ActiveLFOs.Find(LFOId))
+			if (FModulatorBusMixProxy* BusMixes = ActiveBusMixes.Find(static_cast<FBusMixId>(Mix->GetUniqueID())))
 			{
-				LFOProxy->SetFreq(Freq);
+				BusMixes->OnUpdateProxy(InModulator);
 			}
-		};
-
-		IsInAudioThread() ? SetFreq() : FAudioThread::RunCommandOnAudioThread(SetFreq);
+		}
 	}
 } // namespace AudioModulation
 #endif // WITH_AUDIOMODULATION
diff --git a/Engine/Plugins/Runtime/AudioModulation/Source/AudioModulation/Private/AudioModulationInternal.h b/Engine/Plugins/Runtime/AudioModulation/Source/AudioModulation/Private/AudioModulationInternal.h
--- a/Engine/Plugins/Runtime/AudioModulation/Source/AudioModulation/Private/AudioModulationInternal.h
+++ b/Engine/Plugins/Runtime/AudioModulation/Source/AudioModulation/Private/AudioModulationInternal.h
@@ -7,8 +7,8 @@
 
 #include "AudioModulation.h"
 #include "SoundModulatorBase.h"
-#include "SoundModulatorBus.h"
-#include "SoundModulatorBusMix.h"
+#include "SoundControlBus.h"
+#include "SoundControlBusMix.h"
 #include "SoundModulatorLFO.h"
 
 
@@ -23,34 +23,39 @@ namespace AudioModulation
 	class FAudioModulationImpl
 	{
 	public:
+		FAudioModulationImpl();
+
 		void Initialize(const FAudioPluginInitializationParams& InitializationParams);
 
 #if WITH_EDITOR
-		void OnEditSource(const USoundModulationPluginSourceSettingsBase& Settings);
+		void OnEditPluginSettings(const USoundModulationPluginSourceSettingsBase& Settings);
 #endif // WITH_EDITOR
 
-		void OnInitSound(const ModulationSoundId SoundId, const USoundModulationPluginSourceSettingsBase& Settings);
-		void OnInitSource(const uint32 SourceId, const FName& AudioComponentUserId, const uint32 NumChannels, const USoundModulationPluginSourceSettingsBase& Settings);
-		void OnReleaseSound(const ModulationSoundId SoundId, const USoundModulationPluginSourceSettingsBase& Settings);
-		void OnReleaseSource(const uint32 SourceId);
+		void OnInitSound(ISoundModulatable& InSound, const USoundModulationPluginSourceSettingsBase& Settings);
+		void OnInitSource(const uint32 InSourceId, const FName& AudioComponentUserId, const uint32 NumChannels, const USoundModulationPluginSourceSettingsBase& Settings);
+		void OnReleaseSound(ISoundModulatable& InSound);
+		void OnReleaseSource(const uint32 InSourceId);
+
+		void ActivateBus(const USoundControlBusBase& InBus, const ISoundModulatable* Sound = nullptr);
+		void ActivateBusMix(const USoundControlBusMix& InBusMix, const ISoundModulatable* Sound = nullptr);
+		void ActivateLFO(const USoundBusModulatorLFO& InLFO, const ISoundModulatable* Sound = nullptr);
 
-		void ActivateBus(const USoundModulatorBusBase& Bus, const ModulationSoundId SoundId = INDEX_NONE);
-		void ActivateBusMix(const USoundModulatorBusMix& BusMix, bool bReset, const ModulationSoundId SoundId = INDEX_NONE);
-		void ActivateLFO(const USoundModulatorLFO& LFO, const ModulationSoundId SoundId = INDEX_NONE);
-		void DeactivateBusMix(const BusMixId BusMixId, const ModulationSoundId SoundId = INDEX_NONE);
-		void DeactivateBus(const BusId BusId, const ModulationSoundId SoundId = INDEX_NONE);
-		void DeactivateLFO(const LFOId LFOId, const ModulationSoundId SoundId = INDEX_NONE);
+		/**
+		 * Deactivates respective type if InSoundId is invalid (default behavior).  Only deactivates type if
+		 * provided a InSoundId that is releasing, type is set to AutoActivate, and sounds are no longer referencing type.
+		 */
+		void DeactivateBusMix(const FBusMixId InBusMixId, const ISoundModulatable* Sound = nullptr);
+		void DeactivateBus(const FBusId InBusId, const ISoundModulatable* Sound = nullptr);
+		void DeactivateLFO(const FLFOId InLFOId, const ISoundModulatable* Sound = nullptr);
 
-		bool IsBusActive(const BusId BusId) const;
+		bool IsBusActive(const FBusId InBusId) const;
+		bool IsLFOActive(const FLFOId InLFOId) const;
 
 		void ProcessAudio(const FAudioPluginSourceInputData& InputData, FAudioPluginSourceOutputData& OutputData);
-		void ProcessControls(const uint32 SourceId, FSoundModulationControls& Controls);
+		void ProcessControls(const uint32 InSourceId, FSoundModulationControls& OutControls);
 		void ProcessModulators(const float Elapsed);
 
-		void SetBusDefault(const USoundModulatorBusBase& Bus, const float Value);
-		void SetBusRange(const USoundModulatorBusBase& Bus, const FVector2D& Range);
-		void SetBusMixChannel(const USoundModulatorBusMix& BusMix, const USoundModulatorBusBase& Bus, const float TargetValue);
-		void SetLFOFrequency(const USoundModulatorLFO& LFO, const float Freq);
+		void UpdateModulator(const USoundModulatorBase& InModulator);
 
 	private:
 		float CalculateModulationValue(FModulationPatchProxy& Proxy) const;
@@ -59,6 +64,10 @@ namespace AudioModulation
 		BusProxyMap    ActiveBuses;
 		LFOProxyMap    ActiveLFOs;
 
+		// Cache of SoundId to sound setting state at time of sound initialization.
+		TMap<uint32, FModulationSettingsProxy> SoundSettings;
+
+		// Cache of source data while sound is actively playing
 		TArray<FModulationSettingsProxy> SourceSettings;
 
 #if !UE_BUILD_SHIPPING
@@ -69,6 +78,8 @@ namespace AudioModulation
 
 	private:
 		FAudioModulationDebugger Debugger;
+		ISoundModulatable* PreviewSound;
+		FModulationSettingsProxy PreviewSettings;
 #endif // !UE_BUILD_SHIPPING
 	};
 } // namespace AudioModulation
@@ -84,13 +95,13 @@ namespace AudioModulation
 		void Initialize(const FAudioPluginInitializationParams& InitializationParams) { }
 
 #if WITH_EDITOR
-		void OnEditSource(const USoundModulationPluginSourceSettingsBase& Settings) { }
+		void OnEditPluginSettings(const USoundModulationPluginSourceSettingsBase& Settings) { }
 #endif // WITH_EDITOR
 
-		void OnInitSound(const ModulationSoundId SoundId, const USoundModulationPluginSourceSettingsBase& Settings) { }
-		void OnInitSource(const uint32 SourceId, const FName& AudioComponentUserId, const uint32 NumChannels, const USoundModulationPluginSourceSettingsBase& Settings) { }
-		void OnReleaseSound(const ModulationSoundId SoundId, const USoundModulationPluginSourceSettingsBase& Settings) { }
-		void OnReleaseSource(const uint32 SourceId) { }
+		void OnInitSound(ISoundModulatable* InSound, const USoundModulationPluginSourceSettingsBase& Settings) { }
+		void OnInitSource(const uint32 InSourceId, const FName& AudioComponentUserId, const uint32 NumChannels, const USoundModulationPluginSourceSettingsBase& Settings) { }
+		void OnReleaseSound(ISoundModulatable* InSound) { }
+		void OnReleaseSource(const uint32 InSourceId) { }
 
 #if !UE_BUILD_SHIPPING
 		bool OnPostHelp(FCommonViewportClient* ViewportClient, const TCHAR* Stream) { return false; }
@@ -98,22 +109,22 @@ namespace AudioModulation
 		bool OnToggleStat(FCommonViewportClient* ViewportClient, const TCHAR* Stream) { return false; }
 #endif // !UE_BUILD_SHIPPING
 
-		void ActivateBus(const USoundModulatorBusBase& Bus, const ModulationSoundId SoundId = INDEX_NONE) { }
-		void ActivateBusMix(const USoundModulatorBusMix& BusMix, bool bReset, const ModulationSoundId SoundId = INDEX_NONE) { }
-		void ActivateLFO(const USoundModulatorLFO& LFO, const ModulationSoundId SoundId = INDEX_NONE) { }
-		void DeactivateBusMix(const BusMixId BusMixId, const ModulationSoundId SoundId = INDEX_NONE) { }
-		void DeactivateBus(const BusId BusId, const ModulationSoundId SoundId = INDEX_NONE) { }
-		void DeactivateLFO(const LFOId LFOId, const ModulationSoundId SoundId = INDEX_NONE) { }
+		void ActivateBus(const USoundControlBusBase& InBus, const ISoundModulatable* Sound = nullptr) { }
+		void ActivateBusMix(const USoundControlBusMix& InBusMix, const ISoundModulatable* Sound = nullptr) { }
+		void ActivateLFO(const USoundBusModulatorLFO& InLFO, const ISoundModulatable* Sound = nullptr) { }
+		void DeactivateBusMix(const FBusMixId BusMixId, const ISoundModulatable* Sound = nullptr) { }
+		void DeactivateBus(const FBusId BusId, const ISoundModulatable* Sound = nullptr) { }
+		void DeactivateLFO(const FLFOId LFOId, const ISoundModulatable* Sound = nullptr) { }
 
 		void ProcessAudio(const FAudioPluginSourceInputData& InputData, FAudioPluginSourceOutputData& OutputData) { }
-		void ProcessControls(const uint32 SourceId, FSoundModulationControls& Controls) { }
+		void ProcessControls(const uint32 InSourceId, FSoundModulationControls& OutControls) { }
 		void ProcessModulators(const float Elapsed) { }
 
-		void SetBusDefault(const USoundModulatorBusBase& Bus, const float Value) { }
-		void SetBusMin(const USoundModulatorBusBase& Bus, const float Value) { }
-		void SetBusMax(const USoundModulatorBusBase& Bus, const float Value) { }
-		void SetBusMixChannel(const USoundModulatorBusMix& BusMix, const USoundModulatorBusBase& Bus, const float TargetValue) { }
-		void SetLFOFrequency(const USoundModulatorLFO& LFO, const float Freq) { }
+		void SetBusDefault(const USoundControlBusBase& InBus, const float Value) { }
+		void SetBusMin(const USoundControlBusBase& InBus, const float Value) { }
+		void SetBusMax(const USoundControlBusBase& InBus, const float Value) { }
+		void SetBusMixChannel(const USoundControlBusMix& InBusMix, const USoundControlBusBase& InBus, const float TargetValue) { }
+		void SetLFOFrequency(const USoundBusModulatorLFO& InLFO, const float Freq) { }
 	};
 } // namespace AudioModulation
 #endif // WITH_AUDIOMODULATION
diff --git a/Engine/Plugins/Runtime/AudioModulation/Source/AudioModulation/Private/AudioModulationStatics.cpp b/Engine/Plugins/Runtime/AudioModulation/Source/AudioModulation/Private/AudioModulationStatics.cpp
--- a/Engine/Plugins/Runtime/AudioModulation/Source/AudioModulation/Private/AudioModulationStatics.cpp
+++ b/Engine/Plugins/Runtime/AudioModulation/Source/AudioModulation/Private/AudioModulationStatics.cpp
@@ -8,7 +8,7 @@
 #include "AudioModulationLogging.h"
 
 #include "Engine/Engine.h"
-#include "SoundModulatorBus.h"
+#include "SoundControlBus.h"
 
 #define LOCTEXT_NAMESPACE "AudioModulationStatics"
 
@@ -47,7 +47,7 @@ UAudioModulationStatics::UAudioModulationStatics(const FObjectInitializer& Objec
 {
 }
 
-void UAudioModulationStatics::ActivateBus(const UObject* WorldContextObject, USoundModulatorBusBase* Bus)
+void UAudioModulationStatics::ActivateBus(const UObject* WorldContextObject, USoundControlBusBase* Bus)
 {
 	if (!Bus)
 	{
@@ -61,24 +61,24 @@ void UAudioModulationStatics::ActivateBus(const UObject* WorldContextObject, USo
 	}
 }
 
-void UAudioModulationStatics::ActivateBusMix(const UObject* WorldContextObject, USoundModulatorBusMix* BusMix)
+void UAudioModulationStatics::ActivateBusMix(const UObject* WorldContextObject, USoundControlBusMix* BusMix)
 {
 	if (BusMix)
 	{
 		UWorld* World = GetAudioWorld(WorldContextObject);
 		if (AudioModulation::FAudioModulationImpl* ModulationImpl = GetModulationImpl(World))
 		{
-			ModulationImpl->ActivateBusMix(*BusMix, false);
+			ModulationImpl->ActivateBusMix(*BusMix);
 		}
 	}
 }
 
-void UAudioModulationStatics::ActivateModulator(const UObject* WorldContextObject, USoundModulatorBase* Modulator)
+void UAudioModulationStatics::ActivateModulator(const UObject* WorldContextObject, USoundBusModulatorBase* Modulator)
 {
 	UWorld* World = GetAudioWorld(WorldContextObject);
 	if (AudioModulation::FAudioModulationImpl* ModulationImpl = GetModulationImpl(World))
 	{
-		if (USoundModulatorLFO* LFO = Cast<USoundModulatorLFO>(Modulator))
+		if (USoundBusModulatorLFO* LFO = Cast<USoundBusModulatorLFO>(Modulator))
 		{
 			ModulationImpl->ActivateLFO(*LFO);
 		}
@@ -112,10 +112,7 @@ AudioModulation::FAudioModulationImpl* UAudioModulationStatics::GetModulationImp
 	{
 		if (GEngine)
 		{
-			if (FAudioDeviceManager* DeviceManager = GEngine->GetAudioDeviceManager())
-			{
-				AudioDevice = DeviceManager->GetActiveAudioDevice();
-			}
+			AudioDevice = GEngine->GetMainAudioDevice();
 		}
 	}
 
@@ -130,27 +127,27 @@ AudioModulation::FAudioModulationImpl* UAudioModulationStatics::GetModulationImp
 	return nullptr;
 }
 
-USoundVolumeModulatorBus* UAudioModulationStatics::CreateVolumeBus(const UObject* WorldContextObject, FName Name, float DefaultValue, bool Activate)
+USoundVolumeControlBus* UAudioModulationStatics::CreateVolumeBus(const UObject* WorldContextObject, FName Name, float DefaultValue, bool Activate)
 {
-	return CreateModulatorBus<USoundVolumeModulatorBus>(WorldContextObject, Name, DefaultValue, Activate);
+	return CreateModulatorBus<USoundVolumeControlBus>(WorldContextObject, Name, DefaultValue, Activate);
 }
 
-USoundPitchModulatorBus* UAudioModulationStatics::CreatePitchBus(const UObject* WorldContextObject, FName Name, float DefaultValue, bool Activate)
+USoundPitchControlBus* UAudioModulationStatics::CreatePitchBus(const UObject* WorldContextObject, FName Name, float DefaultValue, bool Activate)
 {
-	return CreateModulatorBus<USoundPitchModulatorBus>(WorldContextObject, Name, DefaultValue, Activate);
+	return CreateModulatorBus<USoundPitchControlBus>(WorldContextObject, Name, DefaultValue, Activate);
 }
 
-USoundLPFModulatorBus* UAudioModulationStatics::CreateLPFBus(const UObject* WorldContextObject, FName Name, float DefaultValue, bool Activate)
+USoundLPFControlBus* UAudioModulationStatics::CreateLPFBus(const UObject* WorldContextObject, FName Name, float DefaultValue, bool Activate)
 {
-	return CreateModulatorBus<USoundLPFModulatorBus>(WorldContextObject, Name, DefaultValue, Activate);
+	return CreateModulatorBus<USoundLPFControlBus>(WorldContextObject, Name, DefaultValue, Activate);
 }
 
-USoundHPFModulatorBus* UAudioModulationStatics::CreateHPFBus(const UObject* WorldContextObject, FName Name, float DefaultValue, bool Activate)
+USoundHPFControlBus* UAudioModulationStatics::CreateHPFBus(const UObject* WorldContextObject, FName Name, float DefaultValue, bool Activate)
 {
-	return CreateModulatorBus<USoundHPFModulatorBus>(WorldContextObject, Name, DefaultValue, Activate);
+	return CreateModulatorBus<USoundHPFControlBus>(WorldContextObject, Name, DefaultValue, Activate);
 }
 
-USoundModulatorLFO* UAudioModulationStatics::CreateLFO(const UObject* WorldContextObject, FName Name, float Amplitude, float Frequency, float Offset, bool Activate)
+USoundBusModulatorLFO* UAudioModulationStatics::CreateLFO(const UObject* WorldContextObject, FName Name, float Amplitude, float Frequency, float Offset, bool Activate)
 {
 	UWorld* World = GetAudioWorld(WorldContextObject);
 	if (!World)
@@ -158,7 +155,7 @@ USoundModulatorLFO* UAudioModulationStatics::CreateLFO(const UObject* WorldConte
 		return nullptr;
 	}
 
-	USoundModulatorLFO* NewLFO = NewObject<USoundModulatorLFO>(nullptr, Name, RF_Transient);
+	USoundBusModulatorLFO* NewLFO = NewObject<USoundBusModulatorLFO>(nullptr, Name, RF_Transient);
 	NewLFO->Amplitude = Amplitude;
 	NewLFO->Frequency = Frequency;
 	NewLFO->Offset    = Offset;
@@ -167,14 +164,14 @@ USoundModulatorLFO* UAudioModulationStatics::CreateLFO(const UObject* WorldConte
 	{
 		if (AudioModulation::FAudioModulationImpl* ModulationImpl = GetModulationImpl(World))
 		{
-			ModulationImpl->ActivateLFO(*NewLFO);
+			ModulationImpl->ActivateLFO(*NewLFO, false);
 		}
 	}
 
 	return NewLFO;
 }
 
-USoundModulatorBusMix* UAudioModulationStatics::CreateBusMix(const UObject* WorldContextObject, FName Name, TArray<USoundModulatorBusBase*> Buses, float TargetValue, bool Activate)
+USoundControlBusMix* UAudioModulationStatics::CreateBusMix(const UObject* WorldContextObject, FName Name, TArray<USoundControlBusBase*> Buses, float TargetValue, bool Activate)
 {
 	UWorld* World = GetAudioWorld(WorldContextObject);
 	if (!World)
@@ -182,8 +179,8 @@ USoundModulatorBusMix* UAudioModulationStatics::CreateBusMix(const UObject* Worl
 		return nullptr;
 	}
 
-	USoundModulatorBusMix* NewBusMix = NewObject<USoundModulatorBusMix>(nullptr, Name, RF_Transient);
-	for (USoundModulatorBusBase* Bus : Buses)
+	USoundControlBusMix* NewBusMix = NewObject<USoundControlBusMix>(nullptr, Name, RF_Transient);
+	for (USoundControlBusBase* Bus : Buses)
 	{
 		if (Bus)
 		{
@@ -192,7 +189,7 @@ USoundModulatorBusMix* UAudioModulationStatics::CreateBusMix(const UObject* Worl
 		else
 		{
 			UE_LOG(LogAudioModulation, Warning,
-				TEXT("USoundModulatorBusMix '%s' was created but bus provided is null. Channel not added."),
+				TEXT("USoundControlBusMix '%s' was created but bus provided is null. Channel not added."),
 				*Name.ToString());
 		}
 	}
@@ -208,7 +205,7 @@ USoundModulatorBusMix* UAudioModulationStatics::CreateBusMix(const UObject* Worl
 	return NewBusMix;
 }
 
-void UAudioModulationStatics::DeactivateBus(const UObject* WorldContextObject, USoundModulatorBusBase* Bus)
+void UAudioModulationStatics::DeactivateBus(const UObject* WorldContextObject, USoundControlBusBase* Bus)
 {
 	if (!Bus)
 	{
@@ -218,84 +215,46 @@ void UAudioModulationStatics::DeactivateBus(const UObject* WorldContextObject, U
 	UWorld* World = GetAudioWorld(WorldContextObject);
 	if (AudioModulation::FAudioModulationImpl* ModulationImpl = GetModulationImpl(World))
 	{
-		auto BusId = static_cast<const AudioModulation::BusId>(Bus->GetUniqueID());
+		auto BusId = static_cast<const AudioModulation::FBusId>(Bus->GetUniqueID());
 		ModulationImpl->DeactivateBus(BusId);
 	}
 }
 
-void UAudioModulationStatics::DeactivateBusMix(const UObject* WorldContextObject, USoundModulatorBusMix* BusMix)
+void UAudioModulationStatics::DeactivateBusMix(const UObject* WorldContextObject, USoundControlBusMix* BusMix)
 {
 	if (BusMix)
 	{
 		UWorld* World = GetAudioWorld(WorldContextObject);
 		if (AudioModulation::FAudioModulationImpl* ModulationImpl = GetModulationImpl(World))
 		{
-			auto BusMixId = static_cast<const AudioModulation::BusMixId>(BusMix->GetUniqueID());
+			auto BusMixId = static_cast<const AudioModulation::FBusMixId>(BusMix->GetUniqueID());
 			ModulationImpl->DeactivateBusMix(BusMixId);
 		}
 	}
 }
 
-void UAudioModulationStatics::DeactivateModulator(const UObject* WorldContextObject, USoundModulatorBase* Modulator)
+void UAudioModulationStatics::DeactivateBusModulator(const UObject* WorldContextObject, USoundBusModulatorBase* Modulator)
 {
 	UWorld* World = GetAudioWorld(WorldContextObject);
 	if (AudioModulation::FAudioModulationImpl* ModulationImpl = GetModulationImpl(World))
 	{
-		if (USoundModulatorLFO* LFO = Cast<USoundModulatorLFO>(Modulator))
+		if (USoundBusModulatorLFO* LFO = Cast<USoundBusModulatorLFO>(Modulator))
 		{
-			auto LFOId = static_cast<const AudioModulation::LFOId>(Modulator->GetUniqueID());
+			auto LFOId = static_cast<const AudioModulation::FLFOId>(Modulator->GetUniqueID());
 			ModulationImpl->DeactivateLFO(LFOId);
 		}
 	}
 }
 
-void UAudioModulationStatics::SetBusDefault(const UObject* WorldContextObject, USoundModulatorBusBase* Bus, float Value)
-{
-	if (Bus)
-	{
-		UWorld* World = GetAudioWorld(WorldContextObject);
-		if (AudioModulation::FAudioModulationImpl* ModulationImpl = GetModulationImpl(World))
-		{
-			ModulationImpl->SetBusDefault(*Bus, Value);
-		}
-	}
-	else
-	{
-		UE_LOG(LogAudioModulation, Warning, TEXT("Bus not specified. Bus default set request ignored."));
-	}
-}
-
-void UAudioModulationStatics::SetBusMixChannel(const UObject* WorldContextObject, USoundModulatorBusMix* BusMix, USoundModulatorBusBase* Bus, float TargetValue)
-{
-	if (BusMix && Bus)
-	{
-		UWorld* World = GetAudioWorld(WorldContextObject);
-		if (AudioModulation::FAudioModulationImpl* ModulationImpl = GetModulationImpl(World))
-		{
-			ModulationImpl->SetBusMixChannel(*BusMix, *Bus, TargetValue);
-		}
-	}
-	else
-	{
-		UE_LOG(LogAudioModulation, Warning, TEXT("USoundModulatorBusMix or bus not specified. Bus channel mix set request ignored."));
-	}
-}
-
-void UAudioModulationStatics::SetLFOFrequency(const UObject* WorldContextObject, USoundModulatorLFO* LFO, float Freq)
+void UAudioModulationStatics::UpdateModulator(const UObject* WorldContextObject, USoundModulatorBase* Modulator)
 {
-	if (LFO)
+	if (Modulator)
 	{
 		UWorld* World = GetAudioWorld(WorldContextObject);
 		if (AudioModulation::FAudioModulationImpl* ModulationImpl = GetModulationImpl(World))
 		{
-			ModulationImpl->SetLFOFrequency(*LFO, Freq);
+			ModulationImpl->UpdateModulator(*Modulator);
 		}
 	}
-	else
-	{
-		UE_LOG(LogAudioModulation, Warning,
-			TEXT("USoundModulatorLFO not specified. LFO frequency set ignored."));
-	}
 }
-
 #undef LOCTEXT_NAMESPACE
diff --git a/Engine/Plugins/Runtime/AudioModulation/Source/AudioModulation/Private/AudioModulationStatics.h b/Engine/Plugins/Runtime/AudioModulation/Source/AudioModulation/Private/AudioModulationStatics.h
--- a/Engine/Plugins/Runtime/AudioModulation/Source/AudioModulation/Private/AudioModulationStatics.h
+++ b/Engine/Plugins/Runtime/AudioModulation/Source/AudioModulation/Private/AudioModulationStatics.h
@@ -8,8 +8,8 @@
 
 #include "AudioModulationStatics.generated.h"
 
-class USoundModulatorBus;
-class USoundModulatorBusMix;
+class USoundControlBus;
+class USoundControlBusMix;
 
 namespace AudioModulation
 {
@@ -23,23 +23,33 @@ class UAudioModulationStatics : public UBlueprintFunctionLibrary
 	GENERATED_UCLASS_BODY()
 
 public:
+	/**
+	 * Returns world associated with provided context object
+	 */
+	static UWorld* GetAudioWorld(const UObject* WorldContextObject);
+
+	/**
+	 * Returns modulation implementation associated with the provided world
+	 */
+	static AudioModulation::FAudioModulationImpl* GetModulationImpl(UWorld* World);
+
 	/** Activates a bus. Does nothing if an instance of the provided bus is already active
 	 * @param Bus - Bus to activate
 	 */
-	UFUNCTION(BlueprintCallable, Category = "Audio", DisplayName = "Activate Sound Bus", meta = (WorldContext = "WorldContextObject", Keywords = "activate modulation modulator bus"))
-	static void ActivateBus(const UObject* WorldContextObject, USoundModulatorBusBase* Bus);
+	UFUNCTION(BlueprintCallable, Category = "Audio", DisplayName = "Activate Sound Control Bus", meta = (WorldContext = "WorldContextObject", Keywords = "activate modulation modulator control bus"))
+	static void ActivateBus(const UObject* WorldContextObject, USoundControlBusBase* Bus);
 
 	/** Activates a bus modulator mix. Does nothing if an instance of the provided bus mix is already active
 	 * @param BusMix - Mix to activate
 	 */
-	UFUNCTION(BlueprintCallable, Category = "Audio", DisplayName = "Activate Sound Modulation Bus Mix", meta = (WorldContext = "WorldContextObject", Keywords = "activate modulation modulator bus mix"))
-	static void ActivateBusMix(const UObject* WorldContextObject, USoundModulatorBusMix* BusMix);
+	UFUNCTION(BlueprintCallable, Category = "Audio", DisplayName = "Activate Sound Control Bus Mix", meta = (WorldContext = "WorldContextObject", Keywords = "activate modulation modulator control bus mix"))
+	static void ActivateBusMix(const UObject* WorldContextObject, USoundControlBusMix* BusMix);
 
-	/** Activates a modulator. Does nothing if an instance of the provided modulator is already active
+	/** Activates a bus modulator. Does nothing if an instance of the provided modulator is already active
 	 * @param Modulator - Modulator to activate
 	 */
-	UFUNCTION(BlueprintCallable, Category = "Audio", DisplayName = "Activate Sound Modulator", meta = (WorldContext = "WorldContextObject", Keywords = "activate modulation modulator"))
-	static void ActivateModulator(const UObject* WorldContextObject, USoundModulatorBase* Modulator);
+	UFUNCTION(BlueprintCallable, Category = "Audio", DisplayName = "Activate Sound Modulator", meta = (WorldContext = "WorldContextObject", Keywords = "activate modulation modulator lfo"))
+	static void ActivateModulator(const UObject* WorldContextObject, USoundBusModulatorBase* Modulator);
 
 	/** Creates a volume modulation bus with the provided default value.
 	 * @param Name - Name of bus
@@ -47,7 +57,7 @@ public:
 	 * @param Activate - Whether or not to activate bus on creation.
 	 */
 	UFUNCTION(BlueprintCallable, Category = "Audio", DisplayName = "Create Volume Bus", meta = (WorldContext = "WorldContextObject", Keywords = "create bus modulation volume modulator"))
-	static USoundVolumeModulatorBus* CreateVolumeBus(const UObject* WorldContextObject, FName Name, float DefaultValue, bool Activate);
+	static USoundVolumeControlBus* CreateVolumeBus(const UObject* WorldContextObject, FName Name, float DefaultValue, bool Activate);
 
 	/** Creates a pitch modulation bus with the provided default value.
 	 * @param Name - Name of bus
@@ -55,7 +65,7 @@ public:
 	 * @param Activate - Whether or not to activate bus on creation.
 	 */
 	UFUNCTION(BlueprintCallable, Category = "Audio", DisplayName = "Create Sound Modulator Pitch Bus", meta = (WorldContext = "WorldContextObject", Keywords = "create bus modulation pitch modulator"))
-	static USoundPitchModulatorBus* CreatePitchBus(const UObject* WorldContextObject, FName Name, float DefaultValue, bool Activate);
+	static USoundPitchControlBus* CreatePitchBus(const UObject* WorldContextObject, FName Name, float DefaultValue, bool Activate);
 
 	/** Creates a high-pass filter (HPF) modulation bus with the provided default value.
 	 * @param Name - Name of bus
@@ -63,7 +73,7 @@ public:
 	 * @param Activate - Whether or not to activate bus on creation.
 	 */
 	UFUNCTION(BlueprintCallable, Category = "Audio", DisplayName = "Create Sound Modulator High-Pass Bus", meta = (WorldContext = "WorldContextObject", Keywords = "create bus modulation HPF modulator"))
-	static USoundHPFModulatorBus* CreateHPFBus(const UObject* WorldContextObject, FName Name, float DefaultValue, bool Activate);
+	static USoundHPFControlBus* CreateHPFBus(const UObject* WorldContextObject, FName Name, float DefaultValue, bool Activate);
 
 	/** Creates a low-pass filter (LPF) modulation bus with the provided default value.
 	 * @param Name - Name of bus
@@ -71,7 +81,7 @@ public:
 	 * @param Activate - Whether or not to activate bus on creation.
 	 */
 	UFUNCTION(BlueprintCallable, Category = "Audio", DisplayName = "Create Sound Modulator Low-Pass Bus", meta = (WorldContext = "WorldContextObject", Keywords = "create bus modulation LPF modulator"))
-	static USoundLPFModulatorBus* CreateLPFBus(const UObject* WorldContextObject, FName Name, float DefaultValue, bool Activate);
+	static USoundLPFControlBus* CreateLPFBus(const UObject* WorldContextObject, FName Name, float DefaultValue, bool Activate);
 
 	/** Creates a modulation LFO.
 	 * @param Name - Name of LFO
@@ -81,7 +91,7 @@ public:
 	 * @param Activate - Whether or not to activate lfo on creation.
 	 */
 	UFUNCTION(BlueprintCallable, Category = "Audio", DisplayName = "Create Sound Modulation LFO", meta = (WorldContext = "WorldContextObject", Keywords = "create lfo modulation modulator"))
-	static USoundModulatorLFO* CreateLFO(const UObject* WorldContextObject, FName Name, float Amplitude, float Frequency, float Offset, bool Activate);
+	static USoundBusModulatorLFO* CreateLFO(const UObject* WorldContextObject, FName Name, float Amplitude, float Frequency, float Offset, bool Activate);
 
 	/** Creates a modulation bus mix and adds a bus channel set to the provided target value
 	 * @param Bus - Bus to add to mix
@@ -89,58 +99,30 @@ public:
 	 * @param Activate - Whether or not to activate mix on creation.
 	 */
 	UFUNCTION(BlueprintCallable, Category = "Audio", DisplayName = "Create Sound Modulation Bus Mix", meta = (WorldContext = "WorldContextObject", Keywords = "create bus mix modulation modulator"))
-	static USoundModulatorBusMix* CreateBusMix(const UObject* WorldContextObject, FName Name, TArray<USoundModulatorBusBase*> Buses, float TargetValue, bool Activate);
-
-	/**
-	 * Returns world associated with provided context object
-	 */
-	static UWorld* GetAudioWorld(const UObject* WorldContextObject);
-
-	/**
-	 * Returns modulation implementation associated with the provided world
-	 */
-	static AudioModulation::FAudioModulationImpl* GetModulationImpl(UWorld* World);
+	static USoundControlBusMix* CreateBusMix(const UObject* WorldContextObject, FName Name, TArray<USoundControlBusBase*> Buses, float TargetValue, bool Activate);
 
 	/** Deactivates a bus. Does nothing if an instance of the provided bus is already inactive
 	 * @param Bus - Scope of modulator
 	 */
-	UFUNCTION(BlueprintCallable, Category = "Audio", DisplayName = "Deactivate Sound Modulator Bus", meta = (WorldContext = "WorldContextObject", Keywords = "deactivate modulation modulator bus"))
-	static void DeactivateBus(const UObject* WorldContextObject, USoundModulatorBusBase* Bus);
+	UFUNCTION(BlueprintCallable, Category = "Audio", DisplayName = "Deactivate Sound Modulation Bus", meta = (WorldContext = "WorldContextObject", Keywords = "deactivate modulation modulator bus"))
+	static void DeactivateBus(const UObject* WorldContextObject, USoundControlBusBase* Bus);
 
 	/** Deactivates a modulation bus mix. Does nothing if an instance of the provided bus mix is already inactive
 	 * @param BusMix - Mix to deactivate
 	 */
-	UFUNCTION(BlueprintCallable, Category = "Audio", DisplayName = "Deactivate Sound Modulation Bus Mix", meta = (WorldContext = "WorldContextObject", Keywords = "deactivate modulation modulator"))
-	static void DeactivateBusMix(const UObject* WorldContextObject, USoundModulatorBusMix* BusMix);
+	UFUNCTION(BlueprintCallable, Category = "Audio", DisplayName = "Deactivate Sound Bus Modulation Mix", meta = (WorldContext = "WorldContextObject", Keywords = "deactivate modulation modulator"))
+	static void DeactivateBusMix(const UObject* WorldContextObject, USoundControlBusMix* BusMix);
 
-	/** Deactivates a modulator. Does nothing if an instance of the provided bus mix is already inactive
+	/** Deactivates a bus modulator. Does nothing if an instance of the provided bus mix is already inactive
 	 * @param Modulator - Modulator to activate
 	 * @param Scope - Scope of modulator
 	 */
-	UFUNCTION(BlueprintCallable, Category = "Audio", DisplayName = "Deactivate Sound Modulator", meta = (WorldContext = "WorldContextObject", Keywords = "deactivate modulation modulator"))
-	static void DeactivateModulator(const UObject* WorldContextObject, USoundModulatorBase* Modulator);
-
-	/** Sets a bus's default value
-	 * @param Bus - Bus of which to modify default value
-	 * @param Value - Value for channel to target
-	 */
-	UFUNCTION(BlueprintCallable, Category = "Audio", DisplayName = "Set Bus Default Value", meta = (WorldContext = "WorldContextObject", Keywords = "set bus default modulation modulator"))
-	static void SetBusDefault(const UObject* WorldContextObject, USoundModulatorBusBase* Bus, float Value);
-
-	/** Sets a mix's bus channel(s) to the provided target value
-	 * @param BusMix - BusMix to adjust channel of
-	 * @param Bus - Bus of channel to modify the target value
-	 * @param TargetValue - Value for channel to target
-	 */
-	UFUNCTION(BlueprintCallable, Category = "Audio", DisplayName = "Set Bus Mix Channel", meta = (WorldContext = "WorldContextObject", Keywords = "set bus mix channel modulation modulator"))
-	static void SetBusMixChannel(const UObject* WorldContextObject, USoundModulatorBusMix* BusMix, USoundModulatorBusBase* Bus, float TargetValue);
-
+	UFUNCTION(BlueprintCallable, Category = "Audio", DisplayName = "Deactivate Sound Bus Modulator", meta = (WorldContext = "WorldContextObject", Keywords = "deactivate bus modulation modulator"))
+	static void DeactivateBusModulator(const UObject* WorldContextObject, USoundBusModulatorBase* Modulator);
 
-	/** Sets an LFO's frequency to the provided value
-	 * @param LFO - LFO to adjust frequency
-	 * @param Bus - Bus of channel to modify the target value
-	 * @param TargetValue - Value for channel to target
+	/** Commits updates from a UObject definition of a modulator (e.g. Bus, Bus Mix, LFO) to active instance (does not apply if modulator has not been activated).
+	 * @param Modulator - Modulator to update
 	 */
-	UFUNCTION(BlueprintCallable, Category = "Audio", DisplayName = "Set Modulation LFO Frequency", meta = (WorldContext = "WorldContextObject", Keywords = "set lfo modulation modulator"))
-	static void SetLFOFrequency(const UObject* WorldContextObject, USoundModulatorLFO* LFO, float Freq);
+	UFUNCTION(BlueprintCallable, Category = "Audio", DisplayName = "Update Sound Bus Modulator", meta = (WorldContext = "WorldContextObject", Keywords = "set bus default modulation modulator"))
+	static void UpdateModulator(const UObject* WorldContextObject, USoundModulatorBase* Modulator);
 };
diff --git a/Engine/Plugins/Runtime/AudioModulation/Source/AudioModulation/Private/SoundControlBus.cpp b/Engine/Plugins/Runtime/AudioModulation/Source/AudioModulation/Private/SoundControlBus.cpp
new file mode 100644
--- /dev/null
+++ b/Engine/Plugins/Runtime/AudioModulation/Source/AudioModulation/Private/SoundControlBus.cpp
@@ -0,0 +1,217 @@
+// Copyright 1998-2019 Epic Games, Inc. All Rights Reserved.
+#include "SoundControlBus.h"
+
+#include "AudioDevice.h"
+#include "AudioDeviceManager.h"
+#include "AudioModulation.h"
+#include "AudioModulationInternal.h"
+#include "AudioModulationLogging.h"
+#include "Engine/World.h"
+#include "SoundModulatorLFO.h"
+
+
+USoundControlBusBase::USoundControlBusBase(const FObjectInitializer& ObjectInitializer)
+	: Super(ObjectInitializer)
+	, DefaultValue(1.0f)
+	, Min(0.0f)
+	, Max(1.0f)
+{
+}
+
+USoundVolumeControlBus::USoundVolumeControlBus(const FObjectInitializer& ObjectInitializer)
+	: Super(ObjectInitializer)
+{
+}
+
+USoundPitchControlBus::USoundPitchControlBus(const FObjectInitializer& ObjectInitializer)
+	: Super(ObjectInitializer)
+{
+}
+
+USoundLPFControlBus::USoundLPFControlBus(const FObjectInitializer& ObjectInitializer)
+	: Super(ObjectInitializer)
+{
+}
+
+USoundHPFControlBus::USoundHPFControlBus(const FObjectInitializer& ObjectInitializer)
+	: Super(ObjectInitializer)
+{
+	DefaultValue = 0.0f;
+}
+
+void USoundControlBusBase::BeginDestroy()
+{
+	Super::BeginDestroy();
+
+	UWorld* World = GetWorld();
+	if (!World)
+	{
+		return;
+	}
+
+	FAudioDevice* AudioDevice = World->GetAudioDevice();
+	if (!AudioDevice)
+	{
+		return;
+	}
+
+	check(AudioDevice->IsModulationPluginEnabled());
+	if (IAudioModulation* ModulationInterface = AudioDevice->ModulationInterface.Get())
+	{
+		auto ModulationImpl = static_cast<AudioModulation::FAudioModulation*>(ModulationInterface)->GetImpl();
+		check(ModulationImpl);
+
+		auto BusId = static_cast<const AudioModulation::FBusId>(GetUniqueID());
+		ModulationImpl->DeactivateBus(BusId);
+	}
+}
+
+namespace AudioModulation
+{
+	FControlBusProxy::FControlBusProxy()
+		: TModulatorProxyRefBase<FBusId>()
+		, DefaultValue(0.0f)
+		, LFOValue(1.0f)
+		, MixValue(NAN)
+		, Operator(ESoundModulatorOperator::Multiply)
+		, Range(0.0f, 1.0f)
+	{
+	}
+
+	FControlBusProxy::FControlBusProxy(const USoundControlBusBase& Bus)
+		: TModulatorProxyRefBase<FBusId>(Bus.GetName(), Bus.GetUniqueID(), Bus.bAutoActivate)
+		, DefaultValue(Bus.DefaultValue)
+		, LFOValue(1.0f)
+		, MixValue(NAN)
+		, Operator(Bus.GetOperator())
+		, Range(Bus.Min, Bus.Max)
+	{
+		if (Bus.Min > Bus.Max)
+		{
+			Range.X = Bus.Max;
+			Range.Y = Bus.Min;
+		}
+
+		DefaultValue = FMath::Clamp(DefaultValue, Range.X, Range.Y);
+
+		for (const USoundBusModulatorBase* Modulator: Bus.Modulators)
+		{
+			if (const USoundBusModulatorLFO* LFO = Cast<USoundBusModulatorLFO>(Modulator))
+			{
+				const FLFOId RefLFOId = static_cast<FLFOId>(LFO->GetUniqueID());
+				LFOIds.Add(RefLFOId);
+			}
+		}
+	}
+
+	float FControlBusProxy::GetDefaultValue() const
+	{
+		return DefaultValue;
+	}
+
+	const TArray<FLFOId>& FControlBusProxy::GetLFOIds() const
+	{
+		return LFOIds;
+	}
+
+	float FControlBusProxy::GetLFOValue() const
+	{
+		return LFOValue;
+	}
+
+	FVector2D FControlBusProxy::GetRange() const
+	{
+		return Range;
+	}
+
+	float FControlBusProxy::GetMixValue() const
+	{
+		return MixValue;
+	}
+
+	float FControlBusProxy::GetValue() const
+	{
+		const float DefaultMixed = Mix(DefaultValue);
+		return FMath::Clamp(DefaultMixed * LFOValue, Range.X, Range.Y);
+	}
+
+	float FControlBusProxy::Mix(float ValueA) const
+	{
+		// If mix value is NaN, it is uninitialized (effectively, the parent bus is inactive)
+		// and therefore not mixable, so just return the second value.
+		if (FMath::IsNaN(MixValue))
+		{
+			return ValueA;
+		}
+		return Mix(MixValue, ValueA);
+	}
+		
+	float FControlBusProxy::Mix(float ValueA, float ValueB) const
+	{
+		switch (Operator)
+		{
+			case ESoundModulatorOperator::Min:
+			{
+				return FMath::Min(ValueA, ValueB);
+			}
+			break;
+
+			case ESoundModulatorOperator::Max:
+			{
+				return FMath::Max(ValueA, ValueB);
+			}
+			break;
+
+			case ESoundModulatorOperator::Multiply:
+			{
+				return ValueA * ValueB;
+			}
+			break;
+
+			default:
+			{
+				static_assert(static_cast<int32>(ESoundModulatorOperator::Count) == 3, "Possible missing ESoundModulationOperator case coverage");
+			}
+			break;
+		}
+
+		check(false);
+		return NAN;
+	}
+
+	void FControlBusProxy::MixIn(const float InValue)
+	{
+		MixValue = Mix(InValue);
+	}
+
+	void FControlBusProxy::MixLFO(LFOProxyMap& LFOMap)
+	{
+		for (const FLFOId& LFOId : LFOIds)
+		{
+			if (FModulatorLFOProxy* LFOProxy = LFOMap.Find(LFOId))
+			{
+				LFOValue *= LFOProxy->GetValue();
+			}
+		}
+	}
+
+	void FControlBusProxy::OnUpdateProxy(const USoundModulatorBase& InModulatorArchetype)
+	{
+		const FControlBusProxy CopyProxy(*CastChecked<USoundControlBusBase>(&InModulatorArchetype));
+		auto UpdateProxy = [this, CopyProxy]()
+		{
+			// LFOIds cannot be updated as this would swap sounds subscribing to LFOs if set to AutoActivate
+			DefaultValue = FMath::Clamp(CopyProxy.DefaultValue, CopyProxy.Range.GetMin(), CopyProxy.Range.GetMax());
+			Operator = CopyProxy.Operator;
+			Range = CopyProxy.Range;
+		};
+
+		IsInAudioThread() ? UpdateProxy() : FAudioThread::RunCommandOnAudioThread(UpdateProxy);
+	}
+
+	void FControlBusProxy::Reset()
+	{
+		LFOValue = 1.0f;
+		MixValue = NAN;
+	}
+} // namespace AudioModulation
diff --git a/Engine/Plugins/Runtime/AudioModulation/Source/AudioModulation/Private/SoundControlBus.h b/Engine/Plugins/Runtime/AudioModulation/Source/AudioModulation/Private/SoundControlBus.h
new file mode 100644
--- /dev/null
+++ b/Engine/Plugins/Runtime/AudioModulation/Source/AudioModulation/Private/SoundControlBus.h
@@ -0,0 +1,121 @@
+// Copyright 1998-2019 Epic Games, Inc. All Rights Reserved.
+#pragma once
+
+#include "CoreMinimal.h"
+#include "UObject/ObjectMacros.h"
+#include "UObject/Object.h"
+
+#include "SoundModulatorLFO.h"
+#include "SoundModulationValue.h"
+
+#include "SoundControlBus.generated.h"
+
+
+UENUM()
+enum class ESoundModulatorOperator : uint8
+{
+	/** Multiply all mix values together */
+	Multiply,
+
+	/** Take the lowest mix value */
+	Min,
+
+	/** Take the highest mix value */
+	Max,
+
+	Count UMETA(Hidden),
+};
+
+
+UCLASS(BlueprintType, hidecategories = Object, abstract, MinimalAPI)
+class USoundControlBusBase : public USoundModulatorBase
+{
+	GENERATED_UCLASS_BODY()
+
+public:
+	/** Default value of modulator (when no mix is applied). */
+	UPROPERTY(EditAnywhere, Category = Modulation, BlueprintReadWrite, meta = (UIMin = "0", UIMax = "1"))
+	float DefaultValue;
+
+	/** Minimum value the bus can achieve (applied post mix phase, pre patch output) */
+	UPROPERTY(EditAnywhere, Category = Modulation, BlueprintReadWrite, meta = (UIMin = "0", UIMax = "1"))
+	float Min;
+
+	/** Maximum value the bus can achieve (applied post mix phase, pre patch output) */
+	UPROPERTY(EditAnywhere, Category = Modulation, BlueprintReadWrite, meta = (UIMin = "0", UIMax = "1"))
+	float Max;
+
+	UPROPERTY(EditAnywhere, Category = Modulation, BlueprintReadWrite)
+	TArray<USoundBusModulatorBase*> Modulators;
+
+	virtual void BeginDestroy() override;
+
+	virtual ESoundModulatorOperator GetOperator() const { return ESoundModulatorOperator::Multiply; }
+};
+
+UCLASS(BlueprintType, hidecategories = Object, editinlinenew, MinimalAPI)
+class USoundVolumeControlBus : public USoundControlBusBase
+{
+	GENERATED_UCLASS_BODY()
+};
+
+UCLASS(BlueprintType, hidecategories = Object, editinlinenew, MinimalAPI)
+class USoundPitchControlBus : public USoundControlBusBase
+{
+	GENERATED_UCLASS_BODY()
+
+	virtual ESoundModulatorOperator GetOperator() const { return ESoundModulatorOperator::Multiply; }
+};
+
+UCLASS(BlueprintType, hidecategories = Object, editinlinenew, MinimalAPI)
+class USoundHPFControlBus : public USoundControlBusBase
+{
+	GENERATED_UCLASS_BODY()
+
+	virtual ESoundModulatorOperator GetOperator() const { return ESoundModulatorOperator::Max; }
+};
+
+UCLASS(BlueprintType, hidecategories = Object, editinlinenew, MinimalAPI)
+class USoundLPFControlBus : public USoundControlBusBase
+{
+	GENERATED_UCLASS_BODY()
+
+	virtual ESoundModulatorOperator GetOperator() const { return ESoundModulatorOperator::Min; }
+};
+
+namespace AudioModulation
+{
+	class FControlBusProxy : public TModulatorProxyRefBase<FBusId>
+	{
+	public:
+		FControlBusProxy();
+		FControlBusProxy(const USoundControlBusBase& Bus);
+
+		void OnUpdateProxy(const USoundModulatorBase& InModulatorArchetype) override;
+
+		float GetDefaultValue() const;
+		const TArray<FLFOId>& GetLFOIds() const;
+		float GetLFOValue() const;
+		float GetMixValue() const;
+		FVector2D GetRange() const;
+		float GetValue() const;
+		void MixIn(const float InValue);
+		void MixLFO(LFOProxyMap& LFOMap);
+		void Reset();
+
+	private:
+		float Mix(float ValueA) const;
+		float Mix(float ValueA, float ValueB) const;
+
+		float DefaultValue;
+
+		// Cached values
+		float LFOValue;
+		float MixValue;
+
+		TArray<FLFOId> LFOIds;
+		ESoundModulatorOperator Operator;
+		FVector2D Range;
+	};
+	using BusProxyMap = TMap<FBusId, FControlBusProxy>;
+} // namespace AudioModulation
diff --git a/Engine/Plugins/Runtime/AudioModulation/Source/AudioModulation/Private/SoundControlBusMix.cpp b/Engine/Plugins/Runtime/AudioModulation/Source/AudioModulation/Private/SoundControlBusMix.cpp
new file mode 100644
--- /dev/null
+++ b/Engine/Plugins/Runtime/AudioModulation/Source/AudioModulation/Private/SoundControlBusMix.cpp
@@ -0,0 +1,161 @@
+// Copyright 1998-2019 Epic Games, Inc. All Rights Reserved.
+#include "SoundControlBusMix.h"
+
+#include "AudioDevice.h"
+#include "AudioModulation.h"
+#include "AudioModulationInternal.h"
+#include "AudioModulationLogging.h"
+#include "Engine/World.h"
+
+
+FSoundControlBusMixChannel::FSoundControlBusMixChannel()
+	: Bus(nullptr)
+{
+}
+
+FSoundControlBusMixChannel::FSoundControlBusMixChannel(USoundControlBusBase* InBus, const float TargetValue)
+	: Bus(InBus)
+{
+	Value.TargetValue = FMath::Clamp(TargetValue, 0.0f, 1.0f);
+}
+
+USoundControlBusMix::USoundControlBusMix(const FObjectInitializer& ObjectInitializer)
+	: Super(ObjectInitializer)
+{
+}
+
+void USoundControlBusMix::BeginDestroy()
+{
+	Super::BeginDestroy();
+
+	UWorld* World = GetWorld();
+	if (!World)
+	{
+		return;
+	}
+
+	FAudioDevice* AudioDevice = World->GetAudioDevice();
+	if (!AudioDevice)
+	{
+		return;
+	}
+
+	check(AudioDevice->IsModulationPluginEnabled());
+	if (IAudioModulation* ModulationInterface = AudioDevice->ModulationInterface.Get())
+	{
+		auto ModulationImpl = static_cast<AudioModulation::FAudioModulation*>(ModulationInterface)->GetImpl();
+		check(ModulationImpl);
+
+		auto BusMixId = static_cast<const AudioModulation::FBusMixId>(GetUniqueID());
+		ModulationImpl->DeactivateBusMix(BusMixId);
+	}
+}
+
+namespace AudioModulation
+{
+	FModulatorBusMixChannelProxy::FModulatorBusMixChannelProxy(const FSoundControlBusMixChannel& Channel)
+		: TModulatorProxyBase<FBusId>(Channel.Bus->GetName(), Channel.Bus->GetUniqueID())
+		, Value(Channel.Value)
+	{
+	}
+
+	FModulatorBusMixProxy::FModulatorBusMixProxy(const USoundControlBusMix& Mix)
+		: TModulatorProxyRefBase<FBusMixId>(Mix.GetName(), Mix.GetUniqueID(), Mix.bAutoActivate)
+		, Status(BusMixStatus::Enabled)
+	{
+		for (const FSoundControlBusMixChannel& Channel : Mix.Channels)
+		{
+			if (Channel.Bus)
+			{
+				auto BusId = static_cast<const AudioModulation::FBusId>(Channel.Bus->GetUniqueID());
+#if !UE_BUILD_SHIPPING
+				if (Channels.Contains(BusId))
+				{
+					UE_LOG(LogAudioModulation, Warning,
+						TEXT("USoundControlBusMix '%s' already contains bus '%s'. Only one representative channel for this bus added."),
+						*Mix.GetFullName(), *Channel.Bus->GetFullName());
+				}
+#endif // UE_BUILD_SHIPPING
+				Channels.Emplace(BusId, FModulatorBusMixChannelProxy(Channel));
+			}
+			else
+			{
+				UE_LOG(LogAudioModulation, Warning,
+					TEXT("USoundControlBusMix '%s' has channel with no bus specified. "
+					"Mix activated but channel ignored."),
+					*Mix.GetFullName());
+			}
+		}
+	}
+
+	bool FModulatorBusMixProxy::CanDestroy() const
+	{
+		if (Status != BusMixStatus::Stopped)
+		{
+			return false;
+		}
+
+		return TModulatorProxyRefBase<FBusMixId>::CanDestroy();
+	}
+
+	void FModulatorBusMixProxy::OnUpdateProxy(const USoundModulatorBase& InModulatorArchetype)
+	{
+		const FModulatorBusMixProxy CopyProxy(*CastChecked<USoundControlBusMix>(&InModulatorArchetype));
+		auto UpdateProxy = [this, CopyProxy]()
+		{
+			check(IsInAudioThread());
+
+			Channels = CopyProxy.Channels;
+		};
+
+		IsInAudioThread() ? UpdateProxy() : FAudioThread::RunCommandOnAudioThread(UpdateProxy);
+	}
+
+	void FModulatorBusMixProxy::SetEnabled()
+	{
+		Status = BusMixStatus::Enabled;
+	}
+
+	void FModulatorBusMixProxy::SetStopping()
+	{
+		if (Status == BusMixStatus::Enabled)
+		{
+			Status = BusMixStatus::Stopping;
+		}
+	}
+
+	void FModulatorBusMixProxy::Update(float Elapsed, BusProxyMap& ProxyMap)
+	{
+		bool bRequestStop = true;
+		for (TPair<FBusId, FModulatorBusMixChannelProxy>& Channel : Channels)
+		{
+			FModulatorBusMixChannelProxy& ChannelProxy = Channel.Value;
+			FSoundModulationValue& MixChannelValue = ChannelProxy.Value;
+
+			if (FControlBusProxy* BusProxy = ProxyMap.Find(ChannelProxy.GetId()))
+			{
+				MixChannelValue.Update(Elapsed);
+
+				const float CurrentValue = MixChannelValue.GetCurrentValue();
+				if (Status == BusMixStatus::Stopping)
+				{
+					MixChannelValue.TargetValue = BusProxy->GetDefaultValue();
+					if (!FMath::IsNearlyEqual(MixChannelValue.TargetValue, CurrentValue))
+					{
+						bRequestStop = false;
+					}
+				}
+				else
+				{
+					bRequestStop = false;
+				}
+				BusProxy->MixIn(CurrentValue);
+			}
+		}
+
+		if (bRequestStop)
+		{
+			Status = BusMixStatus::Stopped;
+		}
+	}
+} // namespace AudioModulation
diff --git a/Engine/Plugins/Runtime/AudioModulation/Source/AudioModulation/Private/SoundControlBusMix.h b/Engine/Plugins/Runtime/AudioModulation/Source/AudioModulation/Private/SoundControlBusMix.h
new file mode 100644
--- /dev/null
+++ b/Engine/Plugins/Runtime/AudioModulation/Source/AudioModulation/Private/SoundControlBusMix.h
@@ -0,0 +1,80 @@
+// Copyright 1998-2019 Epic Games, Inc. All Rights Reserved.
+#pragma once
+
+#include "CoreMinimal.h"
+#include "UObject/ObjectMacros.h"
+#include "UObject/Object.h"
+
+#include "SoundControlBus.h"
+#include "SoundModulationValue.h"
+
+#include "SoundControlBusMix.generated.h"
+
+
+USTRUCT(BlueprintType)
+struct FSoundControlBusMixChannel
+{
+	GENERATED_USTRUCT_BODY()
+
+	FSoundControlBusMixChannel();
+	FSoundControlBusMixChannel(USoundControlBusBase* InBus, const float TargetValue);
+
+	/* Bus controlled by channel. */
+	UPROPERTY(EditAnywhere, Category = Channel, BlueprintReadOnly)
+	USoundControlBusBase* Bus;
+
+	/* Value mix is set to. */
+	UPROPERTY(EditAnywhere, Category = Channel, BlueprintReadWrite, meta = (ShowOnlyInnerProperties))
+	FSoundModulationValue Value;
+};
+
+UCLASS(config = Engine, autoexpandcategories = (Channel, Mix), editinlinenew, BlueprintType, MinimalAPI)
+class USoundControlBusMix : public USoundModulatorBase
+{
+	GENERATED_UCLASS_BODY()
+
+public:
+
+	virtual void BeginDestroy() override;
+
+	/* Array of channels controlled by mix. */
+	UPROPERTY(EditAnywhere, Category = Mix, BlueprintReadOnly)
+	TArray<FSoundControlBusMixChannel> Channels;
+};
+
+namespace AudioModulation
+{
+	struct FModulatorBusMixChannelProxy : public TModulatorProxyBase<FBusId>
+	{
+		FModulatorBusMixChannelProxy(const FSoundControlBusMixChannel& Channel);
+		FSoundModulationValue Value;
+	};
+
+	class FModulatorBusMixProxy : public TModulatorProxyRefBase<FBusMixId>
+	{
+	public:
+		enum class BusMixStatus : uint8
+		{
+			Enabled,
+			Stopping,
+			Stopped
+		};
+
+		FModulatorBusMixProxy(const USoundControlBusMix& Mix);
+		virtual ~FModulatorBusMixProxy() = default;
+
+		virtual bool CanDestroy() const override;
+		virtual void OnUpdateProxy(const USoundModulatorBase& InModulatorArchetype) override;
+
+		void SetEnabled();
+		void SetStopping();
+		void Update(const float Elapsed, BusProxyMap& ProxyMap);
+
+		TMap<FBusId, FModulatorBusMixChannelProxy> Channels;
+
+	private:
+		BusMixStatus Status;
+	};
+
+	using BusMixProxyMap = TMap<FBusMixId, FModulatorBusMixProxy>;
+} // namespace AudioModulation
diff --git a/Engine/Plugins/Runtime/AudioModulation/Source/AudioModulation/Private/SoundModulationPatch.cpp b/Engine/Plugins/Runtime/AudioModulation/Source/AudioModulation/Private/SoundModulationPatch.cpp
--- a/Engine/Plugins/Runtime/AudioModulation/Source/AudioModulation/Private/SoundModulationPatch.cpp
+++ b/Engine/Plugins/Runtime/AudioModulation/Source/AudioModulation/Private/SoundModulationPatch.cpp
@@ -11,41 +11,45 @@
 USoundModulationSettings::USoundModulationSettings(const FObjectInitializer& ObjectInitializer)
 	: Super(ObjectInitializer)
 {
+	Volume.Output.Transform.Curve = ESoundModulatorOutputCurve::Log;
+	Volume.Output.Transform.Scalar = 0.5f;
+
 	Lowpass.Output.Transform.OutputMin = 20.0f;
 	Lowpass.Output.Transform.OutputMax = 20000.0f;
+	Lowpass.Output.Transform.Curve = ESoundModulatorOutputCurve::Exp;
 	Lowpass.Output.Operator = ESoundModulatorOperator::Min;
 
 	Highpass.DefaultInputValue = 0.0f;
 	Highpass.Output.Transform.OutputMin = 20.0f;
 	Highpass.Output.Transform.OutputMax = 20000.0f;
+	Highpass.Output.Transform.Curve = ESoundModulatorOutputCurve::Exp;
 	Highpass.Output.Operator = ESoundModulatorOperator::Max;
 }
 
 #if WITH_EDITOR
-void USoundModulationSettings::PostEditChangeProperty(FPropertyChangedEvent& PropertyChangedEvent)
+void USoundModulationSettings::OnPostEditChange(UWorld* World)
 {
-	if (!PropertyChangedEvent.Property)
+	if (AudioModulation::FAudioModulationImpl* Impl = UAudioModulationStatics::GetModulationImpl(World))
 	{
-		return;
+		Impl->OnEditPluginSettings(*this);
 	}
+}
 
-	if (AudioModulation::FAudioModulationImpl* Impl = UAudioModulationStatics::GetModulationImpl(GetWorld()))
+void USoundModulationSettings::PostEditChangeProperty(FPropertyChangedEvent& PropertyChangedEvent)
+{
+	if (PropertyChangedEvent.Property)
 	{
-		Impl->OnEditSource(*this);
+		OnPostEditChange(GetWorld());
 	}
 }
 
 void USoundModulationSettings::PostEditChangeChainProperty(FPropertyChangedChainEvent& PropertyChangedEvent)
 {
-	if (!PropertyChangedEvent.Property)
+	if (PropertyChangedEvent.Property)
 	{
-		return;
+		OnPostEditChange(GetWorld());
 	}
 
-	if (AudioModulation::FAudioModulationImpl* Impl = UAudioModulationStatics::GetModulationImpl(GetWorld()))
-	{
-		Impl->OnEditSource(*this);
-	}
 }
 #endif // WITH_EDITOR
 
@@ -97,9 +101,9 @@ namespace AudioModulation
 		, Transform(Input.Transform)
 		, bSampleAndHold(Input.bSampleAndHold)
 	{
-		if (const USoundModulatorBusBase* Bus = Input.GetBus())
+		if (const USoundControlBusBase* Bus = Input.GetBus())
 		{
-			BusId = static_cast<AudioModulation::BusId>(Bus->GetUniqueID());
+			BusId = static_cast<AudioModulation::FBusId>(Bus->GetUniqueID());
 		}
 	}
 
@@ -131,9 +135,6 @@ namespace AudioModulation
 	}
 
 	FModulationSettingsProxy::FModulationSettingsProxy()
-#if WITH_EDITOR
-		: ObjectID(0)
-#endif // WITH_EDITOR
 	{
 		Lowpass.OutputProxy.Transform.OutputMin  = 20.0f;
 		Lowpass.OutputProxy.Transform.OutputMax  = 20000.0f;
@@ -148,19 +149,17 @@ namespace AudioModulation
 	}
 
 	FModulationSettingsProxy::FModulationSettingsProxy(const USoundModulationSettings& Settings)
-		: Volume(Settings.Volume)
+		: TModulatorProxyBase<uint32>(Settings.GetName(), Settings.GetUniqueID())
+		, Volume(Settings.Volume)
 		, Pitch(Settings.Pitch)
 		, Lowpass(Settings.Lowpass)
 		, Highpass(Settings.Highpass)
-#if WITH_EDITOR
-		, ObjectID(Settings.GetUniqueID())
-#endif // WITH_EDITOR
 	{
-		for (const USoundModulatorBusMix* Mix : Settings.Mixes)
+		for (const USoundControlBusMix* Mix : Settings.Mixes)
 		{
 			if (Mix)
 			{
-				Mixes.Add(static_cast<BusMixId>(Mix->GetUniqueID()));
+				Mixes.Add(static_cast<FBusMixId>(Mix->GetUniqueID()));
 			}
 		}
 	}
diff --git a/Engine/Plugins/Runtime/AudioModulation/Source/AudioModulation/Private/SoundModulationPatch.h b/Engine/Plugins/Runtime/AudioModulation/Source/AudioModulation/Private/SoundModulationPatch.h
--- a/Engine/Plugins/Runtime/AudioModulation/Source/AudioModulation/Private/SoundModulationPatch.h
+++ b/Engine/Plugins/Runtime/AudioModulation/Source/AudioModulation/Private/SoundModulationPatch.h
@@ -6,8 +6,9 @@
 #include "UObject/Object.h"
 #include "SoundModulationTransform.h"
 #include "SoundModulationValue.h"
-#include "SoundModulatorBus.h"
-#include "SoundModulatorBusMix.h"
+#include "SoundModulatorBase.h"
+#include "SoundControlBus.h"
+#include "SoundControlBusMix.h"
 
 #include "SoundModulationPatch.generated.h"
 
@@ -27,7 +28,7 @@ struct FSoundModulationOutput
 	FSoundModulationOutput();
 
 	/** Operator used when mixing input values */
-	UPROPERTY(EditAnywhere, Category = Modulation, BlueprintReadWrite)
+	UPROPERTY(VisibleAnywhere, BlueprintReadOnly, Category = Modulation)
 	ESoundModulatorOperator Operator;
 
 	/** Final transform before passing to output */
@@ -51,7 +52,7 @@ struct FSoundModulationInputBase
 	UPROPERTY(EditAnywhere, BlueprintReadWrite, Category = Input, meta = (ShowOnlyInnerProperties))
 	FSoundModulationInputTransform Transform;
 
-	virtual const USoundModulatorBusBase* GetBus() const PURE_VIRTUAL(FSoundModulationInputBase::GetBus, return nullptr; );
+	virtual const USoundControlBusBase* GetBus() const PURE_VIRTUAL(FSoundModulationInputBase::GetBus, return nullptr; );
 };
 
 USTRUCT(BlueprintType)
@@ -63,9 +64,9 @@ struct FSoundVolumeModulationInput : public FSoundModulationInputBase
 
 	/** The input bus */
 	UPROPERTY(EditAnywhere, BlueprintReadWrite, Category = Input)
-	USoundVolumeModulatorBus* Bus;
+	USoundVolumeControlBus* Bus;
 
-	virtual const USoundModulatorBusBase* GetBus() const { return Cast<USoundModulatorBusBase>(Bus); }
+	virtual const USoundControlBusBase* GetBus() const { return Cast<USoundControlBusBase>(Bus); }
 };
 
 USTRUCT(BlueprintType)
@@ -77,9 +78,9 @@ struct FSoundPitchModulationInput : public FSoundModulationInputBase
 
 	/** The input bus */
 	UPROPERTY(EditAnywhere, BlueprintReadWrite, Category = Input)
-	USoundPitchModulatorBus* Bus;
+	USoundPitchControlBus* Bus;
 
-	virtual const USoundModulatorBusBase* GetBus() const { return Cast<USoundModulatorBusBase>(Bus); }
+	virtual const USoundControlBusBase* GetBus() const { return Cast<USoundControlBusBase>(Bus); }
 };
 
 USTRUCT(BlueprintType)
@@ -91,9 +92,9 @@ struct FSoundLPFModulationInput : public FSoundModulationInputBase
 
 	/** The input bus */
 	UPROPERTY(EditAnywhere, BlueprintReadWrite, Category = Input)
-	USoundLPFModulatorBus* Bus;
+	USoundLPFControlBus* Bus;
 
-	virtual const USoundModulatorBusBase* GetBus() const { return Cast<USoundModulatorBusBase>(Bus); }
+	virtual const USoundControlBusBase* GetBus() const { return Cast<USoundControlBusBase>(Bus); }
 };
 
 USTRUCT(BlueprintType)
@@ -105,9 +106,9 @@ struct FSoundHPFModulationInput : public FSoundModulationInputBase
 
 	/** The input bus */
 	UPROPERTY(EditAnywhere, BlueprintReadWrite, Category = Input)
-	USoundHPFModulatorBus* Bus;
+	USoundHPFControlBus* Bus;
 
-	virtual const USoundModulatorBusBase* GetBus() const { return Cast<USoundModulatorBusBase>(Bus); }
+	virtual const USoundControlBusBase* GetBus() const { return Cast<USoundControlBusBase>(Bus); }
 };
 
 USTRUCT(BlueprintType)
@@ -120,7 +121,7 @@ struct FSoundModulationPatchBase
 	virtual ~FSoundModulationPatchBase() = default;
 
 	/** Default value of patch (value used when all inputs are either not provided or not active) */
-	UPROPERTY(EditAnywhere, BlueprintReadWrite, Category = Inputs)
+	UPROPERTY(EditAnywhere, BlueprintReadWrite, Category = Inputs, meta = (UIMin = "0", UIMax = "1"))
 	float DefaultInputValue;
 
 	/** Final modulation parameters to apply */
@@ -225,9 +226,11 @@ public:
 	// mix will be removed once all sound settings referencing mix stop. Manual
 	// mix activation is ignored if already activated by means of modulation settings.
 	UPROPERTY(EditAnywhere, BlueprintReadWrite, Category = Mixes)
-	TArray<USoundModulatorBusMix*> Mixes;
+	TArray<USoundControlBusMix*> Mixes;
 
 #if WITH_EDITOR
+	void OnPostEditChange(UWorld* World);
+
 	virtual void PostEditChangeProperty(FPropertyChangedEvent& PropertyChangedEvent) override;
 	virtual void PostEditChangeChainProperty(FPropertyChangedChainEvent& PropertyChangedEvent) override;
 #endif
@@ -242,7 +245,7 @@ namespace AudioModulation
 		FModulationInputProxy();
 		FModulationInputProxy(const FSoundModulationInputBase& Patch);
 
-		BusId BusId;
+		FBusId BusId;
 		FSoundModulationInputTransform Transform;
 		uint8 bSampleAndHold : 1;
 	};
@@ -281,7 +284,7 @@ namespace AudioModulation
 		FModulationOutputProxy OutputProxy;
 	};
 
-	struct FModulationSettingsProxy
+	struct FModulationSettingsProxy : TModulatorProxyBase<uint32>
 	{
 		FModulationSettingsProxy();
 		FModulationSettingsProxy(const USoundModulationSettings& Settings);
@@ -291,10 +294,6 @@ namespace AudioModulation
 		FModulationPatchProxy Lowpass;
 		FModulationPatchProxy Highpass;
 
-		TArray<BusMixId> Mixes;
-
-#if WITH_EDITOR
-		uint32 ObjectID;
-#endif // WITH_EDITOR
+		TArray<FBusMixId> Mixes;
 	};
 } // namespace AudioModulation
diff --git a/Engine/Plugins/Runtime/AudioModulation/Source/AudioModulation/Private/SoundModulationTransform.cpp b/Engine/Plugins/Runtime/AudioModulation/Source/AudioModulation/Private/SoundModulationTransform.cpp
--- a/Engine/Plugins/Runtime/AudioModulation/Source/AudioModulation/Private/SoundModulationTransform.cpp
+++ b/Engine/Plugins/Runtime/AudioModulation/Source/AudioModulation/Private/SoundModulationTransform.cpp
@@ -22,7 +22,7 @@ void FSoundModulationInputTransform::Apply(float& Value) const
 FSoundModulationOutputTransform::FSoundModulationOutputTransform()
 	: InputMin(0.0f)
 	, InputMax(1.0f)
-	, Curve(ESoundModulatorOutputCurve::Exp)
+	, Curve(ESoundModulatorOutputCurve::Linear)
 	, Scalar(2.5f)
 	, OutputMin(0.0f)
 	, OutputMax(1.0f)
diff --git a/Engine/Plugins/Runtime/AudioModulation/Source/AudioModulation/Private/SoundModulationTransform.h b/Engine/Plugins/Runtime/AudioModulation/Source/AudioModulation/Private/SoundModulationTransform.h
--- a/Engine/Plugins/Runtime/AudioModulation/Source/AudioModulation/Private/SoundModulationTransform.h
+++ b/Engine/Plugins/Runtime/AudioModulation/Source/AudioModulation/Private/SoundModulationTransform.h
@@ -24,11 +24,11 @@ struct FSoundModulationInputTransform
 	UPROPERTY(EditAnywhere, Category = Modulation, BlueprintReadWrite, meta = (DisplayName = "Input Max", UIMin = "0", UIMax = "1"))
 	float InputMax;
 
-	/** Minimum value to clamp to map the output to. */
+	/** Minimum value to clamp the output to. */
 	UPROPERTY(EditAnywhere, Category = Modulation, BlueprintReadWrite, meta = (DisplayName = "Output Min", UIMin = "0", UIMax = "1"))
 	float OutputMin;
 
-	/** Maximum value to clamp to map the output to. */
+	/** Maximum value to clamp the output to. */
 	UPROPERTY(EditAnywhere, Category = Modulation, BlueprintReadWrite, meta = (DisplayName = "Output Max", UIMin = "0", UIMax = "1"))
 	float OutputMax;
 
@@ -80,11 +80,11 @@ struct FSoundModulationOutputTransform
 	UPROPERTY(EditAnywhere, Category = Curve, BlueprintReadWrite)
 	FRuntimeFloatCurve FloatCurve;
 
-	/** Minimum value to clamp to map the output to. */
+	/** Minimum value to clamp output to. */
 	UPROPERTY(EditAnywhere, Category = Output, BlueprintReadWrite)
 	float OutputMin;
 
-	/** Maximum value to clamp to map the output to. */
+	/** Maximum value to clamp output to. */
 	UPROPERTY(EditAnywhere, Category = Output, BlueprintReadWrite)
 	float OutputMax;
 
diff --git a/Engine/Plugins/Runtime/AudioModulation/Source/AudioModulation/Private/SoundModulationValue.cpp b/Engine/Plugins/Runtime/AudioModulation/Source/AudioModulation/Private/SoundModulationValue.cpp
--- a/Engine/Plugins/Runtime/AudioModulation/Source/AudioModulation/Private/SoundModulationValue.cpp
+++ b/Engine/Plugins/Runtime/AudioModulation/Source/AudioModulation/Private/SoundModulationValue.cpp
@@ -4,9 +4,9 @@
 
 namespace AudioModulation
 {
-	const BusMixId InvalidBusMixId = TNumericLimits<uint32>::Max();
-	const BusId    InvalidBusId    = TNumericLimits<uint32>::Max();
-	const BusId    InvalidLFOId    = TNumericLimits<uint32>::Max();
+	const FBusMixId InvalidBusMixId = INDEX_NONE;
+	const FBusId    InvalidBusId    = INDEX_NONE;
+	const FBusId    InvalidLFOId    = INDEX_NONE;
 } // namespace AudioModulation
 
 FSoundModulationValue::FSoundModulationValue()
diff --git a/Engine/Plugins/Runtime/AudioModulation/Source/AudioModulation/Private/SoundModulationValue.h b/Engine/Plugins/Runtime/AudioModulation/Source/AudioModulation/Private/SoundModulationValue.h
--- a/Engine/Plugins/Runtime/AudioModulation/Source/AudioModulation/Private/SoundModulationValue.h
+++ b/Engine/Plugins/Runtime/AudioModulation/Source/AudioModulation/Private/SoundModulationValue.h
@@ -9,14 +9,14 @@
 
 namespace AudioModulation
 {
-	using BusMixId = uint32;
-	extern const BusMixId InvalidBusMixId;
+	using FBusMixId = uint32;
+	extern const FBusMixId InvalidBusMixId;
 
-	using BusId = uint32;
-	extern const BusId InvalidBusId;
+	using FBusId = uint32;
+	extern const FBusId InvalidBusId;
 
-	using LFOId = uint32;
-	extern const LFOId InvalidLFOId;
+	using FLFOId = uint32;
+	extern const FLFOId InvalidLFOId;
 } // namespace AudioModulation
 
 USTRUCT(BlueprintType)
diff --git a/Engine/Plugins/Runtime/AudioModulation/Source/AudioModulation/Private/SoundModulatorBase.cpp b/Engine/Plugins/Runtime/AudioModulation/Source/AudioModulation/Private/SoundModulatorBase.cpp
--- a/Engine/Plugins/Runtime/AudioModulation/Source/AudioModulation/Private/SoundModulatorBase.cpp
+++ b/Engine/Plugins/Runtime/AudioModulation/Source/AudioModulation/Private/SoundModulatorBase.cpp
@@ -1,8 +1,56 @@
 // Copyright 1998-2019 Epic Games, Inc. All Rights Reserved.
 #include "SoundModulatorBase.h"
 
+#include "AudioDevice.h"
+#include "AudioDeviceManager.h"
+
+#include "AudioModulation.h"
+#include "AudioModulationInternal.h"
+
 
 USoundModulatorBase::USoundModulatorBase(const FObjectInitializer& ObjectInitializer)
 	: Super(ObjectInitializer)
+	, bAutoActivate(0)
+{
+}
+
+#if WITH_EDITOR
+void USoundModulatorBase::PostEditChangeProperty(struct FPropertyChangedEvent& PropertyChangedEvent)
+{
+	Super::PostEditChangeProperty(PropertyChangedEvent);
+
+	if (!PropertyChangedEvent.Property)
+	{
+		return;
+	}
+
+	// Cannot change AutoActivate while modulator is active to avoid leaking sound count anyway,
+	// so early out to avoid passing along a no-op
+	const FName Name = PropertyChangedEvent.Property->GetFName();
+	if (Name == TEXT("bAutoActivate"))
+	{
+		return;
+	}
+
+	FAudioDeviceManager* DeviceManager = GEngine->GetAudioDeviceManager();
+	if (!DeviceManager)
+	{
+		return;
+	}
+
+	const TArray<FAudioDevice*>& Devices = DeviceManager->GetAudioDevices();
+	for (FAudioDevice* Device : Devices)
+	{
+		if (Device && Device->IsModulationPluginEnabled() && Device->ModulationInterface.IsValid())
+		{
+			auto Impl = static_cast<AudioModulation::FAudioModulation*>(Device->ModulationInterface.Get())->GetImpl();
+			Impl->UpdateModulator(*this);
+		}
+	}
+}
+#endif // WITH_EDITOR
+
+USoundBusModulatorBase::USoundBusModulatorBase(const FObjectInitializer& ObjectInitializer)
+	: Super(ObjectInitializer)
 {
 }
diff --git a/Engine/Plugins/Runtime/AudioModulation/Source/AudioModulation/Private/SoundModulatorBase.h b/Engine/Plugins/Runtime/AudioModulation/Source/AudioModulation/Private/SoundModulatorBase.h
--- a/Engine/Plugins/Runtime/AudioModulation/Source/AudioModulation/Private/SoundModulatorBase.h
+++ b/Engine/Plugins/Runtime/AudioModulation/Source/AudioModulation/Private/SoundModulatorBase.h
@@ -6,9 +6,164 @@
 
 #include "SoundModulatorBase.generated.h"
 
+namespace AudioModulation
+{
+	class FAudioModulationImpl;
+}
 
-UCLASS(hidecategories = Object, abstract, MinimalAPI)
+/**
+ * Base class for all modulators
+ */
+UCLASS(hideCategories = Object, abstract, MinimalAPI)
 class USoundModulatorBase : public UObject
+{
+	GENERATED_UCLASS_BODY()
+
+public:
+	/**
+	 * When set, automatically activates/deactivates when sounds referencing asset begin/end playing respectively.
+	 */
+	UPROPERTY(EditAnywhere, Category = General, BlueprintReadWrite)
+	uint8 bAutoActivate : 1;
+
+	bool CanAutoActivate(const ISoundModulatable* InSound) const
+	{
+		if (bAutoActivate)
+		{
+			if (InSound)
+			{
+				return true;
+			}
+		}
+		else
+		{
+			if (!InSound || InSound->IsPreviewSound())
+			{
+				return true;
+			}
+		}
+		
+		return false;
+	}
+
+#if WITH_EDITOR
+	virtual void PostEditChangeProperty(FPropertyChangedEvent& PropertyChangedEvent) override;
+#endif // WITH_EDITOR
+};
+
+/**
+ * Base class for modulators that manipulate control bus values
+ */
+UCLASS(hideCategories = Object, abstract, MinimalAPI)
+class USoundBusModulatorBase : public USoundModulatorBase
 {
 	GENERATED_UCLASS_BODY()
 };
+
+namespace AudioModulation
+{
+	template <typename IdType>
+	class TModulatorProxyBase
+	{
+	private:
+		IdType Id;
+
+#if !UE_BUILD_SHIPPING
+		FString Name;
+#endif // !UE_BUILD_SHIPPING
+
+	public:
+		TModulatorProxyBase()
+			: Id(static_cast<IdType>(0))
+		{
+		}
+
+		TModulatorProxyBase(const FString& InName, const uint32 InObjectId)
+			: Id(static_cast<IdType>(InObjectId))
+#if !UE_BUILD_SHIPPING
+			, Name(InName)
+#endif // !UE_BUILD_SHIPPING
+		{
+		}
+
+		IdType GetId() const
+		{
+			return Id;
+		}
+
+		// FOR DEBUG USE ONLY (Not available in shipped builds):
+		// Provides name of object that generated proxy.
+		const FString& GetName() const
+		{
+#if UE_BUILD_SHIPPING
+			static FString Name;
+#endif // UE_BUILD_SHIPPING
+
+			return Name;
+		}
+	};
+
+	template <typename IdType>
+	class TModulatorProxyRefBase : public TModulatorProxyBase<IdType>
+	{
+	public:
+		TModulatorProxyRefBase()
+			: bAutoActivate(false)
+		{
+		}
+
+		TModulatorProxyRefBase(const FString& Name, const IdType InId, const bool bInAutoActivate)
+			: TModulatorProxyBase<IdType>(Name, InId)
+			, bAutoActivate(bInAutoActivate)
+		{
+		}
+
+		virtual ~TModulatorProxyRefBase() = default;
+
+		virtual bool CanDestroy() const
+		{
+			return !bAutoActivate || (bAutoActivate && RefSounds.Num() == 0);
+		}
+
+		virtual void OnUpdateProxy(const USoundModulatorBase& InModulatorArchetype) { }
+
+		int32 OnReleaseSound(const ISoundModulatable& Sound)
+		{
+			if (bAutoActivate)
+			{
+				const int32 NumRemoved = RefSounds.Remove(&Sound);
+				check(NumRemoved == 1);
+			}
+
+			return RefSounds.Num();
+		}
+
+		bool GetAutoActivate() const
+		{
+			return static_cast<bool>(bAutoActivate);
+		}
+
+		const TArray<const ISoundModulatable*>& GetRefSounds() const
+		{
+			return RefSounds;
+		}
+
+		int32 OnInitSound(const ISoundModulatable& Sound)
+		{
+			// Preview sounds force proxies into being auto-activated
+			// to allow for auditioning with the provided modulation settings.
+			bAutoActivate |= Sound.IsPreviewSound();
+
+			if (bAutoActivate)
+			{
+				RefSounds.AddUnique(&Sound);
+			}
+
+			return RefSounds.Num();
+		}
+
+	private:
+		int8 bAutoActivate : 1;
+		TArray<const ISoundModulatable*> RefSounds;
+	};
+} // AudioModulation
\ No newline at end of file
diff --git a/Engine/Plugins/Runtime/AudioModulation/Source/AudioModulation/Private/SoundModulatorBus.cpp b/Engine/Plugins/Runtime/AudioModulation/Source/AudioModulation/Private/SoundModulatorBus.cpp
deleted file mode 100644
diff --git a/Engine/Plugins/Runtime/AudioModulation/Source/AudioModulation/Private/SoundModulatorBus.h b/Engine/Plugins/Runtime/AudioModulation/Source/AudioModulation/Private/SoundModulatorBus.h
deleted file mode 100644
diff --git a/Engine/Plugins/Runtime/AudioModulation/Source/AudioModulation/Private/SoundModulatorBusMix.cpp b/Engine/Plugins/Runtime/AudioModulation/Source/AudioModulation/Private/SoundModulatorBusMix.cpp
deleted file mode 100644
diff --git a/Engine/Plugins/Runtime/AudioModulation/Source/AudioModulation/Private/SoundModulatorBusMix.h b/Engine/Plugins/Runtime/AudioModulation/Source/AudioModulation/Private/SoundModulatorBusMix.h
deleted file mode 100644
diff --git a/Engine/Plugins/Runtime/AudioModulation/Source/AudioModulation/Private/SoundModulatorLFO.cpp b/Engine/Plugins/Runtime/AudioModulation/Source/AudioModulation/Private/SoundModulatorLFO.cpp
--- a/Engine/Plugins/Runtime/AudioModulation/Source/AudioModulation/Private/SoundModulatorLFO.cpp
+++ b/Engine/Plugins/Runtime/AudioModulation/Source/AudioModulation/Private/SoundModulatorLFO.cpp
@@ -8,18 +8,17 @@
 #include "Engine/World.h"
 
 
-USoundModulatorLFO::USoundModulatorLFO(const FObjectInitializer& ObjectInitializer)
+USoundBusModulatorLFO::USoundBusModulatorLFO(const FObjectInitializer& ObjectInitializer)
 	: Super(ObjectInitializer)
 	, Shape(ESoundModulatorLFOShape::Sine)
 	, Amplitude(0.5f)
 	, Frequency(1.0f)
 	, Offset(0.5f)
 	, bLooping(1)
-	, bAutoActivate(0)
 {
 }
 
-void USoundModulatorLFO::BeginDestroy()
+void USoundBusModulatorLFO::BeginDestroy()
 {
 	Super::BeginDestroy();
 
@@ -41,31 +40,24 @@ void USoundModulatorLFO::BeginDestroy()
 		auto ModulationImpl = static_cast<AudioModulation::FAudioModulation*>(ModulationInterface)->GetImpl();
 		check(ModulationImpl);
 
-		auto LFOId = static_cast<const AudioModulation::LFOId>(GetUniqueID());
+		auto LFOId = static_cast<const AudioModulation::FLFOId>(GetUniqueID());
 		ModulationImpl->DeactivateLFO(LFOId);
 	}
 }
 
+
 namespace AudioModulation
 {
 	FModulatorLFOProxy::FModulatorLFOProxy()
-		: Id(0)
-		, Offset(0.0f)
+		: Offset(0.0f)
 		, Value(0.0f)
-		, bAutoActivate(false)
-		, SoundRefCount(0)
 	{
 	}
 
-	FModulatorLFOProxy::FModulatorLFOProxy(const USoundModulatorLFO& InLFO)
-		: Id(static_cast<LFOId>(InLFO.GetUniqueID()))
-#if !UE_BUILD_SHIPPING
-		, Name(InLFO.GetName())
-#endif // !UE_BUILD_SHIPPING
+	FModulatorLFOProxy::FModulatorLFOProxy(const USoundBusModulatorLFO& InLFO)
+		: TModulatorProxyRefBase<FLFOId>(InLFO.GetName(), InLFO.GetUniqueID(), InLFO.bAutoActivate)
 		, Offset(InLFO.Offset)
 		, Value(0.0f)
-		, bAutoActivate(InLFO.bAutoActivate)
-		, SoundRefCount(0)
 	{
 		LFO.SetGain(InLFO.Amplitude);
 		LFO.SetFrequency(InLFO.Frequency);
@@ -77,46 +69,24 @@ namespace AudioModulation
 		LFO.Start();
 	}
 
-	float FModulatorLFOProxy::GetAmplitude() const
-	{
-		return LFO.GetGain();
-	}
-
-	bool FModulatorLFOProxy::GetAutoActivate() const
-	{
-		return bAutoActivate;
-	}
-
-	float FModulatorLFOProxy::GetFreq() const
-	{
-		return LFO.GetFrequency();
-	}
-
-#if !UE_BUILD_SHIPPING
-	const FString& FModulatorLFOProxy::GetName() const
-	{
-		return Name;
-	}
-#endif // !UE_BUILD_SHIPPING
-
-	LFOId FModulatorLFOProxy::GetId() const
-	{
-		return Id;
-	}
-
-	float FModulatorLFOProxy::GetOffset() const
-	{
-		return Offset;
-	}
-
 	float FModulatorLFOProxy::GetValue() const
 	{
 		return Value;
 	}
 
-	void FModulatorLFOProxy::SetFreq(float InFreq)
+	void FModulatorLFOProxy::OnUpdateProxy(const USoundModulatorBase& InModulatorArchetype)
 	{
-		LFO.SetFrequency(InFreq);
+		const FModulatorLFOProxy CopyProxy(*CastChecked<USoundBusModulatorLFO>(&InModulatorArchetype));
+		auto UpdateProxy = [this, CopyProxy]()
+		{
+			check(IsInAudioThread());
+
+			LFO = CopyProxy.LFO;
+			Offset = CopyProxy.Offset;
+			Value = CopyProxy.Value;
+		};
+
+		IsInAudioThread() ? UpdateProxy() : FAudioThread::RunCommandOnAudioThread(UpdateProxy);
 	}
 
 	void FModulatorLFOProxy::Update(float InElapsed)
@@ -129,18 +99,4 @@ namespace AudioModulation
 			Value = LFO.Generate() + Offset;
 		}
 	}
-
-	int32 FModulatorLFOProxy::DecRefSound()
-	{
-		check(SoundRefCount > 0);
-		SoundRefCount--;
-
-		return SoundRefCount;
-	}
-
-	int32 FModulatorLFOProxy::IncRefSound()
-	{
-		SoundRefCount++;
-		return SoundRefCount;
-	}
 } // namespace AudioModulation
diff --git a/Engine/Plugins/Runtime/AudioModulation/Source/AudioModulation/Private/SoundModulatorLFO.h b/Engine/Plugins/Runtime/AudioModulation/Source/AudioModulation/Private/SoundModulatorLFO.h
--- a/Engine/Plugins/Runtime/AudioModulation/Source/AudioModulation/Private/SoundModulatorLFO.h
+++ b/Engine/Plugins/Runtime/AudioModulation/Source/AudioModulation/Private/SoundModulatorLFO.h
@@ -26,7 +26,7 @@ enum class ESoundModulatorLFOShape : uint8
 };
 
 UCLASS(BlueprintType, hidecategories = Object, editinlinenew, MinimalAPI)
-class USoundModulatorLFO : public USoundModulatorBase
+class USoundBusModulatorLFO : public USoundBusModulatorBase
 {
 	GENERATED_UCLASS_BODY()
 
@@ -52,48 +52,25 @@ public:
 	/** Whether or not to loop the oscillation more than once */
 	UPROPERTY(EditAnywhere, Category = Modulation, BlueprintReadWrite)
 	uint8 bLooping : 1;
-
-	/** Automatically activates/deactivates LFO when sounds referencing LFO asset are playing */
-	UPROPERTY(EditAnywhere, Category = Modulation, BlueprintReadWrite)
-	uint8 bAutoActivate : 1;
 };
 
 namespace AudioModulation
 {
-	class FModulatorLFOProxy
+	class FModulatorLFOProxy : public TModulatorProxyRefBase<FLFOId>
 	{
 	public:
 		FModulatorLFOProxy();
-		FModulatorLFOProxy(const USoundModulatorLFO& InLFO);
-
-		float GetAmplitude() const;
-		bool GetAutoActivate() const;
-		float GetFreq() const;
-		LFOId GetId() const;
-#if !UE_BUILD_SHIPPING
-		const FString& GetName() const;
-#endif // !UE_BUILD_SHIPPING
-		float GetOffset() const;
+		FModulatorLFOProxy(const USoundBusModulatorLFO& InLFO);
+
+		void OnUpdateProxy(const USoundModulatorBase& InModulatorArchetype) override;
+
 		float GetValue() const;
-		void SetFreq(float InFreq);
 		void Update(float InElapsed);
 
-		int32 DecRefSound();
-		int32 IncRefSound();
-
 	private:
-		LFOId Id;
-
-#if !UE_BUILD_SHIPPING
-		FString Name;
-#endif // !UE_BUILD_SHIPPING
-
 		Audio::FLFO LFO;
 		float Offset;
 		float Value;
-
-		uint8 bAutoActivate : 1;
-		int32 SoundRefCount;
 	};
-	using LFOProxyMap = TMap<BusId, FModulatorLFOProxy>;
+	using LFOProxyMap = TMap<FBusId, FModulatorLFOProxy>;
 } // namespace AudioModulation
diff --git a/Engine/Plugins/Runtime/AudioModulation/Source/AudioModulation/Public/AudioModulation.h b/Engine/Plugins/Runtime/AudioModulation/Source/AudioModulation/Public/AudioModulation.h
--- a/Engine/Plugins/Runtime/AudioModulation/Source/AudioModulation/Public/AudioModulation.h
+++ b/Engine/Plugins/Runtime/AudioModulation/Source/AudioModulation/Public/AudioModulation.h
@@ -30,10 +30,10 @@ namespace AudioModulation
 		virtual void Initialize(const FAudioPluginInitializationParams& InitializationParams) override;
 
 #if WITH_EDITOR
-		virtual void OnEditSource(const USoundModulationPluginSourceSettingsBase& Settings) override;
+		void OnEditPluginSettings(const USoundModulationPluginSourceSettingsBase& Settings);
 #endif // WITH_EDITOR
 
-		virtual void OnInitSound(const ModulationSoundId SoundId, const USoundModulationPluginSourceSettingsBase& Settings) override;
+		virtual void OnInitSound(ISoundModulatable& Sound, const USoundModulationPluginSourceSettingsBase& Settings) override;
 		virtual void OnInitSource(const uint32 SourceId, const FName& AudioComponentUserId, const uint32 NumChannels, const USoundModulationPluginSourceSettingsBase& Settings) override;
 
 #if !UE_BUILD_SHIPPING
@@ -42,7 +42,7 @@ namespace AudioModulation
 		virtual bool OnToggleStat(FCommonViewportClient* ViewportClient, const TCHAR* Stream) override;
 #endif // !UE_BUILD_SHIPPING
 
-		virtual void OnReleaseSound(const ModulationSoundId SoundId, const USoundModulationPluginSourceSettingsBase& Settings) override;
+		virtual void OnReleaseSound(ISoundModulatable& Sound) override;
 		virtual void OnReleaseSource(const uint32 SourceId) override;
 		virtual void ProcessControls(const uint32 SourceId, FSoundModulationControls& Controls) override;
 		virtual void ProcessModulators(const float Elapsed) override;
diff --git a/Engine/Plugins/Runtime/AudioModulation/Source/AudioModulationEditor/Private/AssetTypeActions/AssetTypeActions_SoundControlBus.cpp b/Engine/Plugins/Runtime/AudioModulation/Source/AudioModulationEditor/Private/AssetTypeActions/AssetTypeActions_SoundControlBus.cpp
new file mode 100644
--- /dev/null
+++ b/Engine/Plugins/Runtime/AudioModulation/Source/AudioModulationEditor/Private/AssetTypeActions/AssetTypeActions_SoundControlBus.cpp
@@ -0,0 +1,25 @@
+// Copyright 1998-2019 Epic Games, Inc. All Rights Reserved.
+
+#include "AssetTypeActions_SoundControlBus.h"
+#include "SoundControlBus.h"
+
+
+UClass* FAssetTypeActions_SoundVolumeControlBus::GetSupportedClass() const
+{
+	return USoundVolumeControlBus::StaticClass();
+}
+
+UClass* FAssetTypeActions_SoundPitchControlBus::GetSupportedClass() const
+{
+	return USoundPitchControlBus::StaticClass();
+}
+
+UClass* FAssetTypeActions_SoundLPFControlBus::GetSupportedClass() const
+{
+	return USoundLPFControlBus::StaticClass();
+}
+
+UClass* FAssetTypeActions_SoundHPFControlBus::GetSupportedClass() const
+{
+	return USoundHPFControlBus::StaticClass();
+}
diff --git a/Engine/Plugins/Runtime/AudioModulation/Source/AudioModulationEditor/Private/AssetTypeActions/AssetTypeActions_SoundControlBus.h b/Engine/Plugins/Runtime/AudioModulation/Source/AudioModulationEditor/Private/AssetTypeActions/AssetTypeActions_SoundControlBus.h
new file mode 100644
--- /dev/null
+++ b/Engine/Plugins/Runtime/AudioModulation/Source/AudioModulationEditor/Private/AssetTypeActions/AssetTypeActions_SoundControlBus.h
@@ -0,0 +1,46 @@
+// Copyright 1998-2019 Epic Games, Inc. All Rights Reserved.
+
+#pragma once
+
+#include "CoreMinimal.h"
+#include "AssetTypeActions_Base.h"
+
+class FAssetTypeActions_SoundVolumeControlBus : public FAssetTypeActions_Base
+{
+public:
+	// IAssetTypeActions Implementation
+	virtual FText GetName() const override { return NSLOCTEXT("AssetTypeActions", "AssetTypeActions_SoundVolumeControlBus", "Control Bus (Volume)"); }
+	virtual FColor GetTypeColor() const override { return FColor(33, 183, 0); }
+	virtual UClass* GetSupportedClass() const override;
+	virtual uint32 GetCategories() override { return EAssetTypeCategories::Sounds; }
+};
+
+class FAssetTypeActions_SoundPitchControlBus : public FAssetTypeActions_Base
+{
+public:
+	// IAssetTypeActions Implementation
+	virtual FText GetName() const override { return NSLOCTEXT("AssetTypeActions", "AssetTypeActions_SoundPitchControlBus", "Control Bus (Pitch)"); }
+	virtual FColor GetTypeColor() const override { return FColor(181, 21, 0); }
+	virtual UClass* GetSupportedClass() const override;
+	virtual uint32 GetCategories() override { return EAssetTypeCategories::Sounds; }
+};
+
+class FAssetTypeActions_SoundLPFControlBus : public FAssetTypeActions_Base
+{
+public:
+	// IAssetTypeActions Implementation
+	virtual FText GetName() const override { return NSLOCTEXT("AssetTypeActions", "AssetTypeActions_SoundLPFControlBus", "Control Bus (LPF)"); }
+	virtual FColor GetTypeColor() const override { return FColor(0, 156, 183); }
+	virtual UClass* GetSupportedClass() const override;
+	virtual uint32 GetCategories() override { return EAssetTypeCategories::Sounds; }
+};
+
+class FAssetTypeActions_SoundHPFControlBus : public FAssetTypeActions_Base
+{
+public:
+	// IAssetTypeActions Implementation
+	virtual FText GetName() const override { return NSLOCTEXT("AssetTypeActions", "AssetTypeActions_SoundHPFControlBus", "Control Bus (HPF)"); }
+	virtual FColor GetTypeColor() const override { return FColor(94, 237, 183); }
+	virtual UClass* GetSupportedClass() const override;
+	virtual uint32 GetCategories() override { return EAssetTypeCategories::Sounds; }
+};
diff --git a/Engine/Plugins/Runtime/AudioModulation/Source/AudioModulationEditor/Private/AssetTypeActions/AssetTypeActions_SoundControlBusMix.cpp b/Engine/Plugins/Runtime/AudioModulation/Source/AudioModulationEditor/Private/AssetTypeActions/AssetTypeActions_SoundControlBusMix.cpp
new file mode 100644
--- /dev/null
+++ b/Engine/Plugins/Runtime/AudioModulation/Source/AudioModulationEditor/Private/AssetTypeActions/AssetTypeActions_SoundControlBusMix.cpp
@@ -0,0 +1,10 @@
+// Copyright 1998-2019 Epic Games, Inc. All Rights Reserved.
+
+#include "AssetTypeActions_SoundControlBusMix.h"
+#include "SoundControlBusMix.h"
+
+
+UClass* FAssetTypeActions_SoundControlBusMix::GetSupportedClass() const
+{
+	return USoundControlBusMix::StaticClass();
+}
diff --git a/Engine/Plugins/Runtime/AudioModulation/Source/AudioModulationEditor/Private/AssetTypeActions/AssetTypeActions_SoundControlBusMix.h b/Engine/Plugins/Runtime/AudioModulation/Source/AudioModulationEditor/Private/AssetTypeActions/AssetTypeActions_SoundControlBusMix.h
new file mode 100644
--- /dev/null
+++ b/Engine/Plugins/Runtime/AudioModulation/Source/AudioModulationEditor/Private/AssetTypeActions/AssetTypeActions_SoundControlBusMix.h
@@ -0,0 +1,16 @@
+// Copyright 1998-2019 Epic Games, Inc. All Rights Reserved.
+
+#pragma once
+
+#include "CoreMinimal.h"
+#include "AssetTypeActions_Base.h"
+
+class FAssetTypeActions_SoundControlBusMix : public FAssetTypeActions_Base
+{
+public:
+	// IAssetTypeActions Implementation
+	virtual FText GetName() const override { return NSLOCTEXT("AssetTypeActions", "AssetTypeActions_SoundControlBusMix", "Modulation Mix"); }
+	virtual FColor GetTypeColor() const override { return FColor(237, 237, 0); }
+	virtual UClass* GetSupportedClass() const override;
+	virtual uint32 GetCategories() override { return EAssetTypeCategories::Sounds; }
+};
diff --git a/Engine/Plugins/Runtime/AudioModulation/Source/AudioModulationEditor/Private/AssetTypeActions/AssetTypeActions_SoundModulatorBus.cpp b/Engine/Plugins/Runtime/AudioModulation/Source/AudioModulationEditor/Private/AssetTypeActions/AssetTypeActions_SoundModulatorBus.cpp
deleted file mode 100644
diff --git a/Engine/Plugins/Runtime/AudioModulation/Source/AudioModulationEditor/Private/AssetTypeActions/AssetTypeActions_SoundModulatorBus.h b/Engine/Plugins/Runtime/AudioModulation/Source/AudioModulationEditor/Private/AssetTypeActions/AssetTypeActions_SoundModulatorBus.h
deleted file mode 100644
diff --git a/Engine/Plugins/Runtime/AudioModulation/Source/AudioModulationEditor/Private/AssetTypeActions/AssetTypeActions_SoundModulatorBusMix.cpp b/Engine/Plugins/Runtime/AudioModulation/Source/AudioModulationEditor/Private/AssetTypeActions/AssetTypeActions_SoundModulatorBusMix.cpp
deleted file mode 100644
diff --git a/Engine/Plugins/Runtime/AudioModulation/Source/AudioModulationEditor/Private/AssetTypeActions/AssetTypeActions_SoundModulatorBusMix.h b/Engine/Plugins/Runtime/AudioModulation/Source/AudioModulationEditor/Private/AssetTypeActions/AssetTypeActions_SoundModulatorBusMix.h
deleted file mode 100644
diff --git a/Engine/Plugins/Runtime/AudioModulation/Source/AudioModulationEditor/Private/AssetTypeActions/AssetTypeActions_SoundModulatorLFO.cpp b/Engine/Plugins/Runtime/AudioModulation/Source/AudioModulationEditor/Private/AssetTypeActions/AssetTypeActions_SoundModulatorLFO.cpp
--- a/Engine/Plugins/Runtime/AudioModulation/Source/AudioModulationEditor/Private/AssetTypeActions/AssetTypeActions_SoundModulatorLFO.cpp
+++ b/Engine/Plugins/Runtime/AudioModulation/Source/AudioModulationEditor/Private/AssetTypeActions/AssetTypeActions_SoundModulatorLFO.cpp
@@ -6,5 +6,5 @@
 
 UClass* FAssetTypeActions_SoundModulatorLFO::GetSupportedClass() const
 {
-	return USoundModulatorLFO::StaticClass();
+	return USoundBusModulatorLFO::StaticClass();
 }
diff --git a/Engine/Plugins/Runtime/AudioModulation/Source/AudioModulationEditor/Private/AssetTypeActions/AssetTypeActions_SoundModulatorLFO.h b/Engine/Plugins/Runtime/AudioModulation/Source/AudioModulationEditor/Private/AssetTypeActions/AssetTypeActions_SoundModulatorLFO.h
--- a/Engine/Plugins/Runtime/AudioModulation/Source/AudioModulationEditor/Private/AssetTypeActions/AssetTypeActions_SoundModulatorLFO.h
+++ b/Engine/Plugins/Runtime/AudioModulation/Source/AudioModulationEditor/Private/AssetTypeActions/AssetTypeActions_SoundModulatorLFO.h
@@ -9,7 +9,7 @@ class FAssetTypeActions_SoundModulatorLFO : public FAssetTypeActions_Base
 {
 public:
 	// IAssetTypeActions Implementation
-	virtual FText GetName() const override { return NSLOCTEXT("AssetTypeActions", "AssetTypeActions_SoundModulatorLFO", "Modulation LFO"); }
+	virtual FText GetName() const override { return NSLOCTEXT("AssetTypeActions", "AssetTypeActions_SoundModulatorLFO", "Modulator LFO"); }
 	virtual FColor GetTypeColor() const override { return FColor(237, 237, 0); }
 	virtual UClass* GetSupportedClass() const override;
 	virtual uint32 GetCategories() override { return EAssetTypeCategories::Sounds; }
diff --git a/Engine/Plugins/Runtime/AudioModulation/Source/AudioModulationEditor/Private/AudioModulationEditor.cpp b/Engine/Plugins/Runtime/AudioModulation/Source/AudioModulationEditor/Private/AudioModulationEditor.cpp
--- a/Engine/Plugins/Runtime/AudioModulation/Source/AudioModulationEditor/Private/AudioModulationEditor.cpp
+++ b/Engine/Plugins/Runtime/AudioModulation/Source/AudioModulationEditor/Private/AudioModulationEditor.cpp
@@ -2,8 +2,8 @@
 #include "AudioModulationEditor.h"
 
 #include "AssetTypeActions/AssetTypeActions_SoundModulationSettings.h"
-#include "AssetTypeActions/AssetTypeActions_SoundModulatorBus.h"
-#include "AssetTypeActions/AssetTypeActions_SoundModulatorBusMix.h"
+#include "AssetTypeActions/AssetTypeActions_SoundControlBus.h"
+#include "AssetTypeActions/AssetTypeActions_SoundControlBusMix.h"
 #include "AssetTypeActions/AssetTypeActions_SoundModulatorLFO.h"
 #include "Layouts/SoundModulationTransformLayout.h"
 #include "SoundModulationTransform.h"
@@ -12,6 +12,8 @@
 
 namespace
 {
+	static const FName AssetToolsName = TEXT("AssetTools");
+
 	TSharedPtr<FSlateStyleSet> StyleSet;
 } // namespace <>
 
@@ -37,32 +39,50 @@ void FAudioModulationEditorModule::SetIcon(const FString& ClassName)
 void FAudioModulationEditorModule::StartupModule()
 {
 	// Register the audio editor asset type actions
-	IAssetTools& AssetTools = FModuleManager::LoadModuleChecked<FAssetToolsModule>("AssetTools").Get();
-
-	AssetTools.RegisterAssetTypeActions(MakeShareable(new FAssetTypeActions_SoundVolumeModulatorBus));
-	AssetTools.RegisterAssetTypeActions(MakeShareable(new FAssetTypeActions_SoundPitchModulatorBus));
-	AssetTools.RegisterAssetTypeActions(MakeShareable(new FAssetTypeActions_SoundLPFModulatorBus));
-	AssetTools.RegisterAssetTypeActions(MakeShareable(new FAssetTypeActions_SoundHPFModulatorBus));
-	AssetTools.RegisterAssetTypeActions(MakeShareable(new FAssetTypeActions_SoundModulatorBusMix));
-	AssetTools.RegisterAssetTypeActions(MakeShareable(new FAssetTypeActions_SoundModulatorLFO));
-	AssetTools.RegisterAssetTypeActions(MakeShareable(new FAssetTypeActions_SoundModulationSettings));
-
-	SetIcon(TEXT("SoundVolumeModulatorBus"));
-	SetIcon(TEXT("SoundPitchModulatorBus"));
-	SetIcon(TEXT("SoundLPFModulatorBus"));
-	SetIcon(TEXT("SoundHPFModulatorBus"));
-	SetIcon(TEXT("SoundModulatorBusMix"));
-	SetIcon(TEXT("SoundModulatorLFO"));
+	IAssetTools& AssetTools = FModuleManager::LoadModuleChecked<FAssetToolsModule>(AssetToolsName).Get();
+
+	AssetActions.Add(MakeShareable(new FAssetTypeActions_SoundVolumeControlBus));
+	AssetActions.Add(MakeShareable(new FAssetTypeActions_SoundPitchControlBus));
+	AssetActions.Add(MakeShareable(new FAssetTypeActions_SoundHPFControlBus));
+	AssetActions.Add(MakeShareable(new FAssetTypeActions_SoundLPFControlBus));
+	AssetActions.Add(MakeShareable(new FAssetTypeActions_SoundControlBusMix));
+	AssetActions.Add(MakeShareable(new FAssetTypeActions_SoundModulatorLFO));
+	AssetActions.Add(MakeShareable(new FAssetTypeActions_SoundModulationSettings));
+
+	for (TSharedPtr<FAssetTypeActions_Base>& AssetAction : AssetActions)
+	{
+		AssetTools.RegisterAssetTypeActions(AssetAction.ToSharedRef());
+	}
+
+	SetIcon(TEXT("SoundVolumeControlBus"));
+	SetIcon(TEXT("SoundPitchControlBus"));
+	SetIcon(TEXT("SoundHPFControlBus"));
+	SetIcon(TEXT("SoundLPFControlBus"));
+	SetIcon(TEXT("SoundControlBusMix"));
+	SetIcon(TEXT("SoundBusModulatorLFO"));
 	SetIcon(TEXT("SoundModulationSettings"));
 
-	FPropertyEditorModule& PropertyModule = FModuleManager::LoadModuleChecked<FPropertyEditorModule>("PropertyEditor");
-	PropertyModule.RegisterCustomPropertyTypeLayout("SoundModulationOutputTransform", FOnGetPropertyTypeCustomizationInstance::CreateStatic(&FSoundModulationOutputTransformLayoutCustomization::MakeInstance));
+	FPropertyEditorModule& PropertyModule = FModuleManager::LoadModuleChecked<FPropertyEditorModule>(TEXT("PropertyEditor"));
+	PropertyModule.RegisterCustomPropertyTypeLayout("SoundModulationOutputTransform",
+		FOnGetPropertyTypeCustomizationInstance::CreateStatic(
+			&FSoundModulationOutputTransformLayoutCustomization::MakeInstance));
 
 	FSlateStyleRegistry::RegisterSlateStyle(*StyleSet.Get());
 }
 
 void FAudioModulationEditorModule::ShutdownModule()
 {
+	if (FModuleManager::Get().IsModuleLoaded(AssetToolsName))
+	{
+		IAssetTools& AssetTools = FModuleManager::GetModuleChecked<FAssetToolsModule>(AssetToolsName).Get();
+		for (TSharedPtr<FAssetTypeActions_Base>& AssetAction : AssetActions)
+		{
+			AssetTools.UnregisterAssetTypeActions(AssetAction.ToSharedRef());
+		}
+	}
+	AssetActions.Reset();
+
+	FSlateStyleRegistry::UnRegisterSlateStyle(*StyleSet.Get());
 }
 
 
diff --git a/Engine/Plugins/Runtime/AudioModulation/Source/AudioModulationEditor/Private/AudioModulationEditor.h b/Engine/Plugins/Runtime/AudioModulation/Source/AudioModulationEditor/Private/AudioModulationEditor.h
--- a/Engine/Plugins/Runtime/AudioModulation/Source/AudioModulationEditor/Private/AudioModulationEditor.h
+++ b/Engine/Plugins/Runtime/AudioModulation/Source/AudioModulationEditor/Private/AudioModulationEditor.h
@@ -15,4 +15,6 @@ public:
 
 private:
 	void SetIcon(const FString& ClassName);
+
+	TArray<TSharedPtr<FAssetTypeActions_Base>> AssetActions;
 };
diff --git a/Engine/Plugins/Runtime/AudioModulation/Source/AudioModulationEditor/Private/Factories/SoundControlBusFactory.cpp b/Engine/Plugins/Runtime/AudioModulation/Source/AudioModulationEditor/Private/Factories/SoundControlBusFactory.cpp
new file mode 100644
--- /dev/null
+++ b/Engine/Plugins/Runtime/AudioModulation/Source/AudioModulationEditor/Private/Factories/SoundControlBusFactory.cpp
@@ -0,0 +1,61 @@
+// Copyright 1998-2019 Epic Games, Inc. All Rights Reserved.
+
+#include "SoundControlBusFactory.h"
+#include "SoundControlBus.h"
+
+
+USoundVolumeControlBusFactory::USoundVolumeControlBusFactory(const FObjectInitializer& ObjectInitializer)
+	: Super(ObjectInitializer)
+{
+	SupportedClass = USoundVolumeControlBus::StaticClass();
+	bCreateNew     = true;
+	bEditorImport  = false;
+	bEditAfterNew  = true;
+}
+
+UObject* USoundVolumeControlBusFactory::FactoryCreateNew(UClass* Class, UObject* InParent, FName Name, EObjectFlags Flags, UObject* Context, FFeedbackContext* Warn)
+{
+	return NewObject<USoundVolumeControlBus>(InParent, Name, Flags);
+}
+
+USoundPitchControlBusFactory::USoundPitchControlBusFactory(const FObjectInitializer& ObjectInitializer)
+	: Super(ObjectInitializer)
+{
+	SupportedClass = USoundPitchControlBus::StaticClass();
+	bCreateNew = true;
+	bEditorImport = false;
+	bEditAfterNew = true;
+}
+
+UObject* USoundPitchControlBusFactory::FactoryCreateNew(UClass* Class, UObject* InParent, FName Name, EObjectFlags Flags, UObject* Context, FFeedbackContext* Warn)
+{
+	return NewObject<USoundPitchControlBus>(InParent, Name, Flags);
+}
+
+USoundLPFControlBusFactory::USoundLPFControlBusFactory(const FObjectInitializer& ObjectInitializer)
+	: Super(ObjectInitializer)
+{
+	SupportedClass = USoundLPFControlBus::StaticClass();
+	bCreateNew = true;
+	bEditorImport = false;
+	bEditAfterNew = true;
+}
+
+UObject* USoundLPFControlBusFactory::FactoryCreateNew(UClass* Class, UObject* InParent, FName Name, EObjectFlags Flags, UObject* Context, FFeedbackContext* Warn)
+{
+	return NewObject<USoundLPFControlBus>(InParent, Name, Flags);
+}
+
+USoundHPFControlBusFactory::USoundHPFControlBusFactory(const FObjectInitializer& ObjectInitializer)
+	: Super(ObjectInitializer)
+{
+	SupportedClass = USoundHPFControlBus::StaticClass();
+	bCreateNew = true;
+	bEditorImport = false;
+	bEditAfterNew = true;
+}
+
+UObject* USoundHPFControlBusFactory::FactoryCreateNew(UClass* Class, UObject* InParent, FName Name, EObjectFlags Flags, UObject* Context, FFeedbackContext* Warn)
+{
+	return NewObject<USoundHPFControlBus>(InParent, Name, Flags);
+}
\ No newline at end of file
diff --git a/Engine/Plugins/Runtime/AudioModulation/Source/AudioModulationEditor/Private/Factories/SoundControlBusFactory.h b/Engine/Plugins/Runtime/AudioModulation/Source/AudioModulationEditor/Private/Factories/SoundControlBusFactory.h
new file mode 100644
--- /dev/null
+++ b/Engine/Plugins/Runtime/AudioModulation/Source/AudioModulationEditor/Private/Factories/SoundControlBusFactory.h
@@ -0,0 +1,44 @@
+// Copyright 1998-2019 Epic Games, Inc. All Rights Reserved.
+
+//=============================================================================
+// USoundModulationBusFactory
+//=============================================================================
+
+#pragma once
+
+#include "CoreMinimal.h"
+#include "UObject/ObjectMacros.h"
+#include "Factories/Factory.h"
+#include "SoundControlBusFactory.generated.h"
+
+UCLASS(hidecategories=Object, MinimalAPI)
+class USoundVolumeControlBusFactory : public UFactory
+{
+	GENERATED_UCLASS_BODY()
+
+	virtual UObject* FactoryCreateNew(UClass* Class, UObject* InParent, FName Name, EObjectFlags Flags, UObject* Context, FFeedbackContext* Warn) override;
+};
+
+UCLASS(hidecategories = Object, MinimalAPI)
+class USoundPitchControlBusFactory : public UFactory
+{
+	GENERATED_UCLASS_BODY()
+
+	virtual UObject* FactoryCreateNew(UClass* Class, UObject* InParent, FName Name, EObjectFlags Flags, UObject* Context, FFeedbackContext* Warn) override;
+};
+
+UCLASS(hidecategories = Object, MinimalAPI)
+class USoundHPFControlBusFactory : public UFactory
+{
+	GENERATED_UCLASS_BODY()
+
+	virtual UObject* FactoryCreateNew(UClass* Class, UObject* InParent, FName Name, EObjectFlags Flags, UObject* Context, FFeedbackContext* Warn) override;
+};
+
+UCLASS(hidecategories = Object, MinimalAPI)
+class USoundLPFControlBusFactory : public UFactory
+{
+	GENERATED_UCLASS_BODY()
+
+	virtual UObject* FactoryCreateNew(UClass* Class, UObject* InParent, FName Name, EObjectFlags Flags, UObject* Context, FFeedbackContext* Warn) override;
+};
\ No newline at end of file
diff --git a/Engine/Plugins/Runtime/AudioModulation/Source/AudioModulationEditor/Private/Factories/SoundControlBusMixFactory.cpp b/Engine/Plugins/Runtime/AudioModulation/Source/AudioModulationEditor/Private/Factories/SoundControlBusMixFactory.cpp
new file mode 100644
--- /dev/null
+++ b/Engine/Plugins/Runtime/AudioModulation/Source/AudioModulationEditor/Private/Factories/SoundControlBusMixFactory.cpp
@@ -0,0 +1,19 @@
+// Copyright 1998-2019 Epic Games, Inc. All Rights Reserved.
+
+#include "SoundControlBusMixFactory.h"
+#include "SoundControlBusMix.h"
+
+
+USoundControlBusMixFactory::USoundControlBusMixFactory(const FObjectInitializer& ObjectInitializer)
+	: Super(ObjectInitializer)
+{
+	SupportedClass = USoundControlBusMix::StaticClass();
+	bCreateNew     = true;
+	bEditorImport  = false;
+	bEditAfterNew  = true;
+}
+
+UObject* USoundControlBusMixFactory::FactoryCreateNew(UClass* Class, UObject* InParent, FName Name, EObjectFlags Flags, UObject* Context, FFeedbackContext* Warn)
+{
+	return NewObject<USoundControlBusMix>(InParent, Name, Flags);
+}
diff --git a/Engine/Plugins/Runtime/AudioModulation/Source/AudioModulationEditor/Private/Factories/SoundControlBusMixFactory.h b/Engine/Plugins/Runtime/AudioModulation/Source/AudioModulationEditor/Private/Factories/SoundControlBusMixFactory.h
new file mode 100644
--- /dev/null
+++ b/Engine/Plugins/Runtime/AudioModulation/Source/AudioModulationEditor/Private/Factories/SoundControlBusMixFactory.h
@@ -0,0 +1,20 @@
+// Copyright 1998-2019 Epic Games, Inc. All Rights Reserved.
+
+//=============================================================================
+// USoundModulationBusFactory
+//=============================================================================
+
+#pragma once
+
+#include "CoreMinimal.h"
+#include "UObject/ObjectMacros.h"
+#include "Factories/Factory.h"
+#include "SoundControlBusMixFactory.generated.h"
+
+UCLASS(hidecategories=Object, MinimalAPI)
+class USoundControlBusMixFactory : public UFactory
+{
+	GENERATED_UCLASS_BODY()
+
+	virtual UObject* FactoryCreateNew(UClass* Class, UObject* InParent, FName Name, EObjectFlags Flags, UObject* Context, FFeedbackContext* Warn) override;
+};
diff --git a/Engine/Plugins/Runtime/AudioModulation/Source/AudioModulationEditor/Private/Factories/SoundModulatorBusFactory.cpp b/Engine/Plugins/Runtime/AudioModulation/Source/AudioModulationEditor/Private/Factories/SoundModulatorBusFactory.cpp
deleted file mode 100644
diff --git a/Engine/Plugins/Runtime/AudioModulation/Source/AudioModulationEditor/Private/Factories/SoundModulatorBusFactory.h b/Engine/Plugins/Runtime/AudioModulation/Source/AudioModulationEditor/Private/Factories/SoundModulatorBusFactory.h
deleted file mode 100644
diff --git a/Engine/Plugins/Runtime/AudioModulation/Source/AudioModulationEditor/Private/Factories/SoundModulatorBusMixFactory.cpp b/Engine/Plugins/Runtime/AudioModulation/Source/AudioModulationEditor/Private/Factories/SoundModulatorBusMixFactory.cpp
deleted file mode 100644
diff --git a/Engine/Plugins/Runtime/AudioModulation/Source/AudioModulationEditor/Private/Factories/SoundModulatorBusMixFactory.h b/Engine/Plugins/Runtime/AudioModulation/Source/AudioModulationEditor/Private/Factories/SoundModulatorBusMixFactory.h
deleted file mode 100644
diff --git a/Engine/Plugins/Runtime/AudioModulation/Source/AudioModulationEditor/Private/Factories/SoundModulatorLFOFactory.cpp b/Engine/Plugins/Runtime/AudioModulation/Source/AudioModulationEditor/Private/Factories/SoundModulatorLFOFactory.cpp
--- a/Engine/Plugins/Runtime/AudioModulation/Source/AudioModulationEditor/Private/Factories/SoundModulatorLFOFactory.cpp
+++ b/Engine/Plugins/Runtime/AudioModulation/Source/AudioModulationEditor/Private/Factories/SoundModulatorLFOFactory.cpp
@@ -7,7 +7,7 @@
 USoundModulatorLFOFactory::USoundModulatorLFOFactory(const FObjectInitializer& ObjectInitializer)
 	: Super(ObjectInitializer)
 {
-	SupportedClass = USoundModulatorLFO::StaticClass();
+	SupportedClass = USoundBusModulatorLFO::StaticClass();
 	bCreateNew     = true;
 	bEditorImport  = false;
 	bEditAfterNew  = true;
@@ -15,5 +15,5 @@ USoundModulatorLFOFactory::USoundModulatorLFOFactory(const FObjectInitializer& O
 
 UObject* USoundModulatorLFOFactory::FactoryCreateNew(UClass* Class, UObject* InParent, FName Name, EObjectFlags Flags, UObject* Context, FFeedbackContext* Warn)
 {
-	return NewObject<USoundModulatorLFO>(InParent, Name, Flags);
+	return NewObject<USoundBusModulatorLFO>(InParent, Name, Flags);
 }
diff --git a/Engine/Source/Runtime/Engine/Classes/Sound/SoundBase.h b/Engine/Source/Runtime/Engine/Classes/Sound/SoundBase.h
--- a/Engine/Source/Runtime/Engine/Classes/Sound/SoundBase.h
+++ b/Engine/Source/Runtime/Engine/Classes/Sound/SoundBase.h
@@ -98,8 +98,8 @@ public:
 	TEnumAsByte<EMaxConcurrentResolutionRule::Type> MaxConcurrentResolutionRule_DEPRECATED;
 #endif // WITH_EDITORONLY_DATA
 
-	/** Number of times this sound is currently being played (counted if sound is virtualized). */
-	int32 CurrentPlayCount;
+	/** Map of device handle to number of times this sound is currently being played using that device(counted if sound is virtualized). */
+	TMap<Audio::FDeviceId, int32> CurrentPlayCount;
 
 #if WITH_EDITORONLY_DATA
 	/** If Override Concurrency is false, the sound concurrency settings to use for this sound. */
diff --git a/Engine/Source/Runtime/Engine/Classes/Sound/SoundClass.h b/Engine/Source/Runtime/Engine/Classes/Sound/SoundClass.h
--- a/Engine/Source/Runtime/Engine/Classes/Sound/SoundClass.h
+++ b/Engine/Source/Runtime/Engine/Classes/Sound/SoundClass.h
@@ -2,6 +2,7 @@
 #pragma once
 
 #include "CoreMinimal.h"
+#include "IAudioExtensionPlugin.h"
 #include "UObject/ObjectMacros.h"
 #include "UObject/Object.h"
 #if WITH_EDITOR
@@ -197,6 +198,13 @@ class ENGINE_API USoundClass : public UObject
 	UPROPERTY(EditAnywhere, BlueprintReadOnly, Category=SoundClass)
 	TArray<struct FPassiveSoundMixModifier> PassiveSoundMixModifiers;
 
+	/**
+	  * Modulation for the sound class. If not set on sound directly, settings
+	  * fall back to the modulation settings provided here.
+	  */
+	UPROPERTY(EditAnywhere, BlueprintReadOnly, Category = Modulation)
+	FSoundModulation Modulation;
+
 public:
 	UPROPERTY()
 	USoundClass* ParentClass;
diff --git a/Engine/Source/Runtime/Engine/Private/ActiveSound.cpp b/Engine/Source/Runtime/Engine/Private/ActiveSound.cpp
--- a/Engine/Source/Runtime/Engine/Private/ActiveSound.cpp
+++ b/Engine/Source/Runtime/Engine/Private/ActiveSound.cpp
@@ -196,6 +196,21 @@ void FActiveSound::AddReferencedObjects( FReferenceCollector& Collector)
 	}
 }
 
+int32 FActiveSound::GetPlayCount() const
+{
+	if (!Sound || !AudioDevice)
+	{
+		return 0;
+	}
+
+	if (const int32* PlayCount = Sound->CurrentPlayCount.Find(AudioDevice->DeviceHandle))
+	{
+		return *PlayCount;
+	}
+
+	return 0;
+}
+
 void FActiveSound::SetPitch(float Value)
 {
 	PitchMultiplier = Value;
@@ -357,6 +372,14 @@ void FActiveSound::SetSourceBusSend(EBusSendType BusSendType, const FSoundSource
 	SoundSourceBusSendsOverride[(int32)BusSendType].Add(SourceBusSendInfo);
 }
 
+void FActiveSound::Stop()
+{
+	if (AudioDevice)
+	{
+		AudioDevice->AddSoundToStop(this);
+	}
+}
+
 void FActiveSound::GetSoundSubmixSends(TArray<FSoundSubmixSendInfo>& OutSends) const
 {
 	if (Sound)
@@ -739,9 +762,17 @@ USoundModulationPluginSourceSettingsBase* FActiveSound::FindModulationSettings()
 {
 	if (UClass* PluginClass = GetAudioPluginCustomSettingsClass(EAudioPlugin::MODULATION))
 	{
-		if (Sound->Modulation.Settings.Num() > 0)
+		for (USoundModulationPluginSourceSettingsBase* Settings : Sound->Modulation.Settings)
 		{
-			for (USoundModulationPluginSourceSettingsBase* Settings : Sound->Modulation.Settings)
+			if (Settings && Settings->IsA(PluginClass))
+			{
+				return Settings;
+			}
+		}
+
+		if (UAudioComponent* AudioComponent = UAudioComponent::GetAudioComponentFromID(AudioComponentID))
+		{
+			for (USoundModulationPluginSourceSettingsBase* Settings : AudioComponent->Modulation.Settings)
 			{
 				if (Settings && Settings->IsA(PluginClass))
 				{
@@ -749,9 +780,10 @@ USoundModulationPluginSourceSettingsBase* FActiveSound::FindModulationSettings()
 				}
 			}
 		}
-		else if (UAudioComponent* AudioComponent = UAudioComponent::GetAudioComponentFromID(AudioComponentID))
+
+		if (USoundClass* SoundClass = Sound->GetSoundClass())
 		{
-			for (USoundModulationPluginSourceSettingsBase* Settings : AudioComponent->Modulation.Settings)
+			for (USoundModulationPluginSourceSettingsBase* Settings : SoundClass->Modulation.Settings)
 			{
 				if (Settings && Settings->IsA(PluginClass))
 				{
@@ -764,7 +796,7 @@ USoundModulationPluginSourceSettingsBase* FActiveSound::FindModulationSettings()
 	return nullptr;
 }
 
-void FActiveSound::Stop(bool bStopNow)
+void FActiveSound::MarkPendingDestroy(bool bDestroyNow)
 {
 	check(AudioDevice);
 
@@ -772,16 +804,21 @@ void FActiveSound::Stop(bool bStopNow)
 
 	if (Sound && !bIsStopping)
 	{
-		Sound->CurrentPlayCount = FMath::Max(Sound->CurrentPlayCount - 1, 0);
-		if (Sound->CurrentPlayCount == 0)
+		int32* PlayCount = AudioDevice ? Sound->CurrentPlayCount.Find(AudioDevice->DeviceHandle) : nullptr;
+		if (PlayCount)
 		{
-			if (AudioDevice->ModulationInterface)
+			*PlayCount = FMath::Max(*PlayCount - 1, 0);
+			if (*PlayCount == 0)
 			{
-				if (USoundModulationPluginSourceSettingsBase* ModulationSettings = FindModulationSettings())
-				{
-					const ModulationSoundId SoundId = static_cast<ModulationSoundId>(Sound->GetUniqueID());
- 					AudioDevice->ModulationInterface->OnReleaseSound(SoundId, *ModulationSettings);
-				}
+				Sound->CurrentPlayCount.Remove(AudioDevice->DeviceHandle);
+			}
+		}
+
+		if (AudioDevice->IsModulationPluginEnabled() && AudioDevice->ModulationInterface.IsValid())
+		{
+			if (USoundModulationPluginSourceSettingsBase* ModulationSettings = FindModulationSettings())
+			{
+				AudioDevice->ModulationInterface->OnReleaseSound(static_cast<ISoundModulatable&>(*this));
 			}
 		}
 	}
@@ -798,7 +835,7 @@ void FActiveSound::Stop(bool bStopNow)
 			bool bStopped = false;
 			if (AudioDevice->IsAudioMixerEnabled() && AudioDevice->IsStoppingVoicesEnabled())
 			{
-				if (bStopNow || !AudioDevice->GetNumFreeSources())
+				if (bDestroyNow || !AudioDevice->GetNumFreeSources())
 				{
 					Source->StopNow();
 					bStopped = true;
@@ -848,7 +885,7 @@ void FActiveSound::Stop(bool bStopNow)
 		RemoveWaveInstance(WaveInstance->WaveInstanceHash);
 	}
 
-	if (bStopNow)
+	if (bDestroyNow)
 	{
 		bIsStopping = false;
 	}
diff --git a/Engine/Source/Runtime/Engine/Private/AudioDevice.cpp b/Engine/Source/Runtime/Engine/Private/AudioDevice.cpp
--- a/Engine/Source/Runtime/Engine/Private/AudioDevice.cpp
+++ b/Engine/Source/Runtime/Engine/Private/AudioDevice.cpp
@@ -205,7 +205,7 @@ FAudioDevice::FAudioDevice()
 	, CommonAudioPoolSize(0)
 	, CommonAudioPool(nullptr)
 	, CommonAudioPoolFreeBytes(0)
-	, DeviceHandle(INDEX_NONE)
+	, DeviceHandle(static_cast<Audio::FDeviceId>(INDEX_NONE))
 	, SpatializationPluginInterface(nullptr)
 	, ReverbPluginInterface(nullptr)
 	, OcclusionInterface(nullptr)
@@ -592,6 +592,11 @@ int32 FAudioDevice::GetMaxSources() const
 	return MaxSources + NumStoppingSources;
 }
 
+TRange<float> FAudioDevice::GetGlobalPitchRange() const
+{
+	return TRange<float>(GlobalMinPitch, GlobalMaxPitch);
+}
+
 void FAudioDevice::Teardown()
 {
 	// Make sure we process any pending game thread tasks before tearing down the audio device.
@@ -4127,7 +4132,7 @@ void FAudioDevice::SendUpdateResultsToGameThread(const int32 FirstActiveIndex)
 {
 	DECLARE_CYCLE_STAT(TEXT("FGameThreadAudioTask.AudioSendResults"), STAT_AudioSendResults, STATGROUP_TaskGraphTasks);
 
-	const uint32 AudioDeviceID = DeviceHandle;
+	const Audio::FDeviceId AudioDeviceID = DeviceHandle;
 	UReverbEffect* ReverbEffect = Effects ? Effects->GetCurrentReverbEffect() : nullptr;
 	FAudioThread::RunCommandOnGameThread([AudioDeviceID, ReverbEffect]()
 	{
@@ -4155,10 +4160,9 @@ void FAudioDevice::StopAllSounds(bool bShouldStopUISounds)
 	{
 		DECLARE_CYCLE_STAT(TEXT("FAudioThreadTask.StopAllSounds"), STAT_AudioStopAllSounds, STATGROUP_AudioThreadCommands);
 
-		FAudioDevice* AudioDevice = this;
-		FAudioThread::RunCommandOnAudioThread([AudioDevice, bShouldStopUISounds]()
+		FAudioThread::RunCommandOnAudioThread([this, bShouldStopUISounds]()
 		{
-			AudioDevice->StopAllSounds(bShouldStopUISounds);
+			StopAllSounds(bShouldStopUISounds);
 		}, GET_STATID(STAT_AudioStopAllSounds));
 
 		return;
@@ -4347,12 +4351,18 @@ void FAudioDevice::AddNewActiveSoundInternal(const FActiveSound& NewActiveSound,
 		}
 	}
 
-	++Sound->CurrentPlayCount;
-	if (ModulationInterface && Sound->CurrentPlayCount == 1)
+	int32* PlayCount = Sound->CurrentPlayCount.Find(DeviceHandle);
+	if (!PlayCount)
+	{
+		PlayCount = &Sound->CurrentPlayCount.Add(DeviceHandle);
+	}
+	(*PlayCount)++;
+
+	if (ModulationInterface.IsValid())
 	{
 		if (USoundModulationPluginSourceSettingsBase* ModulationSettings = ActiveSound->FindModulationSettings())
 		{
-			ModulationInterface->OnInitSound(static_cast<ModulationSoundId>(Sound->GetUniqueID()), *ModulationSettings);
+			ModulationInterface->OnInitSound(static_cast<ISoundModulatable&>(*ActiveSound), *ModulationSettings);
 		}
 	}
 
@@ -4546,7 +4556,7 @@ void FAudioDevice::ProcessingPendingActiveSoundStops(bool bForceDelete)
 		}
 		else
 		{
-			ActiveSound->Stop(bForceDelete);
+			ActiveSound->MarkPendingDestroy(bForceDelete);
 
 			USoundBase* Sound = ActiveSound->GetSound();
 
diff --git a/Engine/Source/Runtime/Engine/Public/ActiveSound.h b/Engine/Source/Runtime/Engine/Public/ActiveSound.h
--- a/Engine/Source/Runtime/Engine/Public/ActiveSound.h
+++ b/Engine/Source/Runtime/Engine/Public/ActiveSound.h
@@ -11,6 +11,7 @@
 #include "AudioDynamicParameter.h"
 #include "Components/AudioComponent.h"
 #include "DSP/VolumeFader.h"
+#include "IAudioExtensionPlugin.h"
 #include "Sound/AudioVolume.h"
 #include "Sound/SoundConcurrency.h"
 #include "Sound/SoundSubmix.h"
@@ -242,7 +243,7 @@ struct FSoundParseParameters
 	}
 };
 
-struct ENGINE_API FActiveSound
+struct ENGINE_API FActiveSound : public ISoundModulatable
 {
 public:
 
@@ -267,6 +268,13 @@ private:
 
 
 public:
+	// ISoundModulatable Implementation
+	USoundModulationPluginSourceSettingsBase* FindModulationSettings() const override;
+	uint32 GetObjectId() const override { return Sound ? Sound->GetUniqueID() : INDEX_NONE; }
+	int32 GetPlayCount() const override;
+	bool IsPreviewSound() const override { return bIsPreviewSound; }
+	void Stop() override;
+
 
 	uint64 GetAudioComponentID() const { return AudioComponentID; }
 	FName GetAudioComponentUserID() const { return AudioComponentUserID; }
@@ -670,7 +678,7 @@ private:
 
 	struct FAsyncTraceDetails
 	{
-		uint32 AudioDeviceID;
+		Audio::FDeviceId AudioDeviceID;
 		FActiveSound* ActiveSound;
 	};
 
@@ -684,15 +692,15 @@ private:
 	/** This is a friend so the audio device can call Stop() on the active sound. */
 	friend class FAudioDevice;
 
-	/** Stops the active sound. Can only be called from the owning audio device. */
-	void Stop(bool bStopNow);
+	/**
+	  * Marks the active sound as pending delete and begins termination of internal resources.
+	  * Only to be called from the owning audio device.
+	  */
+	void MarkPendingDestroy (bool bDestroyNow);
 
 	/** Whether or not the active sound is stopping. */
 	bool IsStopping() const { return bIsStopping; }
 
-	/** Find the active modulation settings */
-	USoundModulationPluginSourceSettingsBase* FindModulationSettings() const;
-
 	/** Called when an active sound has been stopped but needs to update it's stopping sounds. Returns true when stopping sources have finished stopping. */
 	bool UpdateStoppingSources(uint64 CurrentTick, bool bEnsureStopped);
 
diff --git a/Engine/Source/Runtime/Engine/Public/Audio.h b/Engine/Source/Runtime/Engine/Public/Audio.h
--- a/Engine/Source/Runtime/Engine/Public/Audio.h
+++ b/Engine/Source/Runtime/Engine/Public/Audio.h
@@ -132,6 +132,14 @@ struct FWaveInstance;
 struct FReverbSettings;
 struct FSampleLoop;
 
+namespace Audio
+{
+	/**
+	 * Typed identifier for Audio Device Id
+	 */
+	using FDeviceId = uint32;
+}
+
 enum ELoopingMode
 {
 	/** One shot sound */
diff --git a/Engine/Source/Runtime/Engine/Public/AudioDevice.h b/Engine/Source/Runtime/Engine/Public/AudioDevice.h
--- a/Engine/Source/Runtime/Engine/Public/AudioDevice.h
+++ b/Engine/Source/Runtime/Engine/Public/AudioDevice.h
@@ -563,6 +563,11 @@ public:
 	  * including the number of stopping voices reserved. */
 	int32 GetMaxSources() const;
 
+	/**
+	 * Returns global pitch range
+	 */
+	TRange<float> GetGlobalPitchRange() const;
+
 	/**
 	* Stops any sound sources which are using the given buffer.
 	*
@@ -1630,7 +1635,7 @@ public:
 	int32 CommonAudioPoolFreeBytes;
 
 	/** The handle for this audio device used in the audio device manager. */
-	uint32 DeviceHandle;
+	Audio::FDeviceId DeviceHandle;
 
 	/** 3rd party audio spatialization interface. */
 	TAudioSpatializationPtr SpatializationPluginInterface;
diff --git a/Engine/Source/Runtime/Engine/Public/IAudioExtensionPlugin.h b/Engine/Source/Runtime/Engine/Public/IAudioExtensionPlugin.h
--- a/Engine/Source/Runtime/Engine/Public/IAudioExtensionPlugin.h
+++ b/Engine/Source/Runtime/Engine/Public/IAudioExtensionPlugin.h
@@ -472,6 +472,37 @@ public:
 	TArray<USoundModulationPluginSourceSettingsBase*> Settings;
 };
 
+class ISoundModulatable
+{
+public:
+	virtual ~ISoundModulatable() = default;
+
+	/**
+	 * Returns the modulation settings of the sound
+	 */
+	virtual USoundModulationPluginSourceSettingsBase* FindModulationSettings() const = 0;
+
+	/**
+	 * Gets the object definition id of the given playing sound's instance
+	 */
+	virtual uint32 GetObjectId() const = 0;
+
+	/**
+	 * Returns number of actively instances of sound playing (including virtualized instances)
+	 */
+	virtual int32 GetPlayCount() const = 0;
+
+	/**
+	 * Returns whether or not sound is an editor preview sound
+	 */
+	virtual bool IsPreviewSound() const = 0;
+
+	/**
+	 * Stops sound.
+	 */
+	virtual void Stop() = 0;
+};
+
 /************************************************************************/
 /* IAudioModulationFactory                                              */
 /*                                                                      */
@@ -521,9 +552,6 @@ struct FSoundModulationControls
 	}
 };
 
-/** Identifier active sound types currently playing (or virtualized) */
-using ModulationSoundId = uint32;
-
 class IAudioModulation
 {
 public:
@@ -533,13 +561,8 @@ public:
 	/** Initialize the modulation plugin with the same rate and number of sources */
 	virtual void Initialize(const FAudioPluginInitializationParams& InitializationParams) { }
 
-#if WITH_EDITOR
-	/** Allows modulation plugin to respond to settings updates when client is editing settings */
-	virtual void OnEditSource(const USoundModulationPluginSourceSettingsBase& Settings) { }
-#endif // WITH_EDITOR
-
 	/** Called when a USoundBase type begins playing a sound */
-	virtual void OnInitSound(const ModulationSoundId SoundId, const USoundModulationPluginSourceSettingsBase& Settings) { }
+	virtual void OnInitSound(ISoundModulatable& Sound, const USoundModulationPluginSourceSettingsBase& Settings) { }
 
 	/** Called when a source is assigned to a voice */
 	virtual void OnInitSource(const uint32 SourceId, const FName& AudioComponentUserId, const uint32 NumChannels, const USoundModulationPluginSourceSettingsBase& Settings) { }
@@ -548,7 +571,7 @@ public:
 	virtual void OnReleaseSource(const uint32 SourceId) { }
 
 	/** Called when a USoundBase type stops playing any sounds */
-	virtual void OnReleaseSound(const ModulationSoundId SoundId, const USoundModulationPluginSourceSettingsBase& Settings) { }
+	virtual void OnReleaseSound(ISoundModulatable& Sound) { }
 
 #if !UE_BUILD_SHIPPING
 	/** Request to post help from active plugin (non-shipping builds only) */
