Clarification of type name changes in respective BP modulation static call comments/display names

diff --git a/Engine/Plugins/Runtime/AudioModulation/Source/AudioModulation/Private/AudioModulationStatics.cpp b/Engine/Plugins/Runtime/AudioModulation/Source/AudioModulation/Private/AudioModulationStatics.cpp
--- a/Engine/Plugins/Runtime/AudioModulation/Source/AudioModulation/Private/AudioModulationStatics.cpp
+++ b/Engine/Plugins/Runtime/AudioModulation/Source/AudioModulation/Private/AudioModulationStatics.cpp
@@ -73,7 +73,7 @@ void UAudioModulationStatics::ActivateBusMix(const UObject* WorldContextObject,
 	}
 }
 
-void UAudioModulationStatics::ActivateModulator(const UObject* WorldContextObject, USoundBusModulatorBase* Modulator)
+void UAudioModulationStatics::ActivateBusModulator(const UObject* WorldContextObject, USoundBusModulatorBase* Modulator)
 {
 	UWorld* World = GetAudioWorld(WorldContextObject);
 	if (AudioModulation::FAudioModulationImpl* ModulationImpl = GetModulationImpl(World))
diff --git a/Engine/Plugins/Runtime/AudioModulation/Source/AudioModulation/Private/AudioModulationStatics.h b/Engine/Plugins/Runtime/AudioModulation/Source/AudioModulation/Private/AudioModulationStatics.h
--- a/Engine/Plugins/Runtime/AudioModulation/Source/AudioModulation/Private/AudioModulationStatics.h
+++ b/Engine/Plugins/Runtime/AudioModulation/Source/AudioModulation/Private/AudioModulationStatics.h
@@ -45,11 +45,11 @@ public:
 	UFUNCTION(BlueprintCallable, Category = "Audio", DisplayName = "Activate Sound Control Bus Mix", meta = (WorldContext = "WorldContextObject", Keywords = "activate modulation modulator control bus mix"))
 	static void ActivateBusMix(const UObject* WorldContextObject, USoundControlBusMix* BusMix);
 
-	/** Activates a bus modulator. Does nothing if an instance of the provided modulator is already active
+	/** Activates a bus modulator (eg. LFO). Does nothing if an instance of the provided modulator is already active
 	 * @param Modulator - Modulator to activate
 	 */
-	UFUNCTION(BlueprintCallable, Category = "Audio", DisplayName = "Activate Sound Modulator", meta = (WorldContext = "WorldContextObject", Keywords = "activate modulation modulator lfo"))
-	static void ActivateModulator(const UObject* WorldContextObject, USoundBusModulatorBase* Modulator);
+	UFUNCTION(BlueprintCallable, Category = "Audio", DisplayName = "Activate Sound Bus Modulator", meta = (WorldContext = "WorldContextObject", Keywords = "activate modulation modulator lfo"))
+	static void ActivateBusModulator(const UObject* WorldContextObject, USoundBusModulatorBase* Modulator);
 
 	/** Creates a volume modulation bus with the provided default value.
 	 * @param Name - Name of bus
@@ -64,7 +64,7 @@ public:
 	 * @param DefaultValue - Default value for created bus
 	 * @param Activate - Whether or not to activate bus on creation.
 	 */
-	UFUNCTION(BlueprintCallable, Category = "Audio", DisplayName = "Create Sound Modulator Pitch Bus", meta = (WorldContext = "WorldContextObject", Keywords = "create bus modulation pitch modulator"))
+	UFUNCTION(BlueprintCallable, Category = "Audio", DisplayName = "Create Sound Pitch Control Bus", meta = (WorldContext = "WorldContextObject", Keywords = "create bus modulation pitch modulator"))
 	static USoundPitchControlBus* CreatePitchBus(const UObject* WorldContextObject, FName Name, float DefaultValue, bool Activate);
 
 	/** Creates a high-pass filter (HPF) modulation bus with the provided default value.
@@ -72,7 +72,7 @@ public:
 	 * @param DefaultValue - Default value for created bus
 	 * @param Activate - Whether or not to activate bus on creation.
 	 */
-	UFUNCTION(BlueprintCallable, Category = "Audio", DisplayName = "Create Sound Modulator High-Pass Bus", meta = (WorldContext = "WorldContextObject", Keywords = "create bus modulation HPF modulator"))
+	UFUNCTION(BlueprintCallable, Category = "Audio", DisplayName = "Create Sound High-Pass Control Bus", meta = (WorldContext = "WorldContextObject", Keywords = "create bus modulation HPF modulator"))
 	static USoundHPFControlBus* CreateHPFBus(const UObject* WorldContextObject, FName Name, float DefaultValue, bool Activate);
 
 	/** Creates a low-pass filter (LPF) modulation bus with the provided default value.
@@ -80,7 +80,7 @@ public:
 	 * @param DefaultValue - Default value for created bus
 	 * @param Activate - Whether or not to activate bus on creation.
 	 */
-	UFUNCTION(BlueprintCallable, Category = "Audio", DisplayName = "Create Sound Modulator Low-Pass Bus", meta = (WorldContext = "WorldContextObject", Keywords = "create bus modulation LPF modulator"))
+	UFUNCTION(BlueprintCallable, Category = "Audio", DisplayName = "Create Sound Low-Pass Control Bus", meta = (WorldContext = "WorldContextObject", Keywords = "create bus modulation LPF modulator"))
 	static USoundLPFControlBus* CreateLPFBus(const UObject* WorldContextObject, FName Name, float DefaultValue, bool Activate);
 
 	/** Creates a modulation LFO.
@@ -90,7 +90,7 @@ public:
 	 * @param Offset - Offset of new LFO.
 	 * @param Activate - Whether or not to activate lfo on creation.
 	 */
-	UFUNCTION(BlueprintCallable, Category = "Audio", DisplayName = "Create Sound Modulation LFO", meta = (WorldContext = "WorldContextObject", Keywords = "create lfo modulation modulator"))
+	UFUNCTION(BlueprintCallable, Category = "Audio", DisplayName = "Create Sound Modulation Bus LFO", meta = (WorldContext = "WorldContextObject", Keywords = "create lfo modulation modulator"))
 	static USoundBusModulatorLFO* CreateLFO(const UObject* WorldContextObject, FName Name, float Amplitude, float Frequency, float Offset, bool Activate);
 
 	/** Creates a modulation bus mix and adds a bus channel set to the provided target value
@@ -98,7 +98,7 @@ public:
 	 * @param TargetValue - value for added bus channel to target
 	 * @param Activate - Whether or not to activate mix on creation.
 	 */
-	UFUNCTION(BlueprintCallable, Category = "Audio", DisplayName = "Create Sound Modulation Bus Mix", meta = (WorldContext = "WorldContextObject", Keywords = "create bus mix modulation modulator"))
+	UFUNCTION(BlueprintCallable, Category = "Audio", DisplayName = "Create Sound Control Bus Mix", meta = (WorldContext = "WorldContextObject", Keywords = "create bus mix modulation modulator"))
 	static USoundControlBusMix* CreateBusMix(const UObject* WorldContextObject, FName Name, TArray<USoundControlBusBase*> Buses, float TargetValue, bool Activate);
 
 	/** Deactivates a bus. Does nothing if an instance of the provided bus is already inactive
@@ -123,6 +123,6 @@ public:
 	/** Commits updates from a UObject definition of a modulator (e.g. Bus, Bus Mix, LFO) to active instance (does not apply if modulator has not been activated).
 	 * @param Modulator - Modulator to update
 	 */
-	UFUNCTION(BlueprintCallable, Category = "Audio", DisplayName = "Update Sound Bus Modulator", meta = (WorldContext = "WorldContextObject", Keywords = "set bus default modulation modulator"))
+	UFUNCTION(BlueprintCallable, Category = "Audio", DisplayName = "Update Sound Modulator", meta = (WorldContext = "WorldContextObject", Keywords = "update set control bus mix modulation modulator"))
 	static void UpdateModulator(const UObject* WorldContextObject, USoundModulatorBase* Modulator);
 };
