Fix vs2019 build: Missing static cast to base action type.  Move to utility function in implementation file for convenience.

diff --git a/Engine/Plugins/Runtime/AudioModulation/Source/AudioModulationEditor/Private/AudioModulationEditor.cpp b/Engine/Plugins/Runtime/AudioModulation/Source/AudioModulationEditor/Private/AudioModulationEditor.cpp
--- a/Engine/Plugins/Runtime/AudioModulation/Source/AudioModulationEditor/Private/AudioModulationEditor.cpp
+++ b/Engine/Plugins/Runtime/AudioModulation/Source/AudioModulationEditor/Private/AudioModulationEditor.cpp
@@ -15,6 +15,15 @@ namespace
 	static const FName AssetToolsName = TEXT("AssetTools");
 
 	TSharedPtr<FSlateStyleSet> StyleSet;
+
+	template <typename T>
+	void AddAssetAction(IAssetTools& AssetTools, TArray<TSharedPtr<FAssetTypeActions_Base>>& AssetArray)
+	{
+		TSharedPtr<T> AssetAction = MakeShareable(new T);
+		TSharedPtr<FAssetTypeActions_Base> AssetActionBase = StaticCastSharedPtr<FAssetTypeActions_Base>(AssetAction);
+		AssetTools.RegisterAssetTypeActions(AssetAction.ToSharedRef());
+		AssetArray.Add(AssetActionBase);
+	}
 } // namespace <>
 
 FAudioModulationEditorModule::FAudioModulationEditorModule()
@@ -41,18 +50,13 @@ void FAudioModulationEditorModule::StartupModule()
 	// Register the audio editor asset type actions
 	IAssetTools& AssetTools = FModuleManager::LoadModuleChecked<FAssetToolsModule>(AssetToolsName).Get();
 
-	AssetActions.Add(MakeShareable(new FAssetTypeActions_SoundVolumeControlBus));
-	AssetActions.Add(MakeShareable(new FAssetTypeActions_SoundPitchControlBus));
-	AssetActions.Add(MakeShareable(new FAssetTypeActions_SoundHPFControlBus));
-	AssetActions.Add(MakeShareable(new FAssetTypeActions_SoundLPFControlBus));
-	AssetActions.Add(MakeShareable(new FAssetTypeActions_SoundControlBusMix));
-	AssetActions.Add(MakeShareable(new FAssetTypeActions_SoundModulatorLFO));
-	AssetActions.Add(MakeShareable(new FAssetTypeActions_SoundModulationSettings));
-
-	for (TSharedPtr<FAssetTypeActions_Base>& AssetAction : AssetActions)
-	{
-		AssetTools.RegisterAssetTypeActions(AssetAction.ToSharedRef());
-	}
+	AddAssetAction<FAssetTypeActions_SoundVolumeControlBus>(AssetTools, AssetActions);
+	AddAssetAction<FAssetTypeActions_SoundPitchControlBus>(AssetTools, AssetActions);
+	AddAssetAction<FAssetTypeActions_SoundHPFControlBus>(AssetTools, AssetActions);
+	AddAssetAction<FAssetTypeActions_SoundLPFControlBus>(AssetTools, AssetActions);
+	AddAssetAction<FAssetTypeActions_SoundControlBusMix>(AssetTools, AssetActions);
+	AddAssetAction<FAssetTypeActions_SoundModulatorLFO>(AssetTools, AssetActions);
+	AddAssetAction<FAssetTypeActions_SoundModulationSettings>(AssetTools, AssetActions);
 
 	SetIcon(TEXT("SoundVolumeControlBus"));
 	SetIcon(TEXT("SoundPitchControlBus"));
diff --git a/Engine/Plugins/Runtime/AudioModulation/Source/AudioModulationEditor/Private/AudioModulationEditor.h b/Engine/Plugins/Runtime/AudioModulation/Source/AudioModulationEditor/Private/AudioModulationEditor.h
--- a/Engine/Plugins/Runtime/AudioModulation/Source/AudioModulationEditor/Private/AudioModulationEditor.h
+++ b/Engine/Plugins/Runtime/AudioModulation/Source/AudioModulationEditor/Private/AudioModulationEditor.h
@@ -2,6 +2,8 @@
 #pragma once
 
 #include "CoreMinimal.h"
+
+#include "AssetTypeActions_Base.h"
 #include "Modules/ModuleManager.h"
 
 
