Fix missed API update callsite

diff --git a/Engine/Plugins/Runtime/AudioModulation/Source/AudioModulation/Private/AudioModulationStatics.cpp b/Engine/Plugins/Runtime/AudioModulation/Source/AudioModulation/Private/AudioModulationStatics.cpp
--- a/Engine/Plugins/Runtime/AudioModulation/Source/AudioModulation/Private/AudioModulationStatics.cpp
+++ b/Engine/Plugins/Runtime/AudioModulation/Source/AudioModulation/Private/AudioModulationStatics.cpp
@@ -164,7 +164,7 @@ USoundBusModulatorLFO* UAudioModulationStatics::CreateLFO(const UObject* WorldCo
 	{
 		if (AudioModulation::FAudioModulationImpl* ModulationImpl = GetModulationImpl(World))
 		{
-			ModulationImpl->ActivateLFO(*NewLFO, false);
+			ModulationImpl->ActivateLFO(*NewLFO);
 		}
 	}
 
