Mirror proxy defaults and add missing include

diff --git a/Engine/Plugins/Runtime/AudioModulation/Source/AudioModulation/Private/SoundModulationPatch.cpp b/Engine/Plugins/Runtime/AudioModulation/Source/AudioModulation/Private/SoundModulationPatch.cpp
--- a/Engine/Plugins/Runtime/AudioModulation/Source/AudioModulation/Private/SoundModulationPatch.cpp
+++ b/Engine/Plugins/Runtime/AudioModulation/Source/AudioModulation/Private/SoundModulationPatch.cpp
@@ -11,7 +11,7 @@
 USoundModulationSettings::USoundModulationSettings(const FObjectInitializer& ObjectInitializer)
 	: Super(ObjectInitializer)
 {
-	Volume.Output.Transform.Curve = ESoundModulatorOutputCurve::Log;
+	Volume.Output.Transform.Curve = ESoundModulatorOutputCurve::Exp;
 	Volume.Output.Transform.Scalar = 0.5f;
 
 	Lowpass.Output.Transform.OutputMin = 20.0f;
@@ -136,15 +136,18 @@ namespace AudioModulation
 
 	FModulationSettingsProxy::FModulationSettingsProxy()
 	{
-		Lowpass.OutputProxy.Transform.OutputMin  = 20.0f;
-		Lowpass.OutputProxy.Transform.OutputMax  = 20000.0f;
-		Lowpass.OutputProxy.SampleAndHoldValue   = 20000.0f;
+		Volume.OutputProxy.Transform.Curve = ESoundModulatorOutputCurve::Exp;
+		Volume.OutputProxy.Transform.Scalar = 0.5f;
+
+		Lowpass.OutputProxy.Transform.OutputMin = 20.0f;
+		Lowpass.OutputProxy.Transform.OutputMax = 20000.0f;
+		Lowpass.OutputProxy.Transform.Curve = ESoundModulatorOutputCurve::Exp;
 		Lowpass.OutputProxy.Operator = ESoundModulatorOperator::Min;
 
 		Highpass.DefaultInputValue = 0.0f;
 		Highpass.OutputProxy.Transform.OutputMin = 20.0f;
 		Highpass.OutputProxy.Transform.OutputMax = 20000.0f;
-		Highpass.OutputProxy.SampleAndHoldValue  = 20.0f;
+		Highpass.OutputProxy.Transform.Curve = ESoundModulatorOutputCurve::Exp;
 		Highpass.OutputProxy.Operator = ESoundModulatorOperator::Max;
 	}
 
