Trivial: Fix additional missed API callsite change.

diff --git a/Engine/Plugins/Runtime/AudioModulation/Source/AudioModulation/Private/AudioModulationStatics.cpp b/Engine/Plugins/Runtime/AudioModulation/Source/AudioModulation/Private/AudioModulationStatics.cpp
--- a/Engine/Plugins/Runtime/AudioModulation/Source/AudioModulation/Private/AudioModulationStatics.cpp
+++ b/Engine/Plugins/Runtime/AudioModulation/Source/AudioModulation/Private/AudioModulationStatics.cpp
@@ -198,7 +198,7 @@ USoundControlBusMix* UAudioModulationStatics::CreateBusMix(const UObject* WorldC
 	{
 		if (AudioModulation::FAudioModulationImpl* ModulationImpl = GetModulationImpl(World))
 		{
-			ModulationImpl->ActivateBusMix(*NewBusMix, false);
+			ModulationImpl->ActivateBusMix(*NewBusMix);
 		}
 	}
 
