Fix crash when previewing sounds that are using modulation system

diff --git a/Engine/Plugins/Runtime/AudioModulation/Source/AudioModulation/Private/AudioModulationImpl.cpp b/Engine/Plugins/Runtime/AudioModulation/Source/AudioModulation/Private/AudioModulationImpl.cpp
--- a/Engine/Plugins/Runtime/AudioModulation/Source/AudioModulation/Private/AudioModulationImpl.cpp
+++ b/Engine/Plugins/Runtime/AudioModulation/Source/AudioModulation/Private/AudioModulationImpl.cpp
@@ -282,6 +282,13 @@ namespace AudioModulation
 		{
 			DeactivateBus(Input.BusId, &InSound);
 		}
+
+		if (InSound.IsPreviewSound())
+		{
+			check(&InSound == PreviewSound);
+			PreviewSound = nullptr;
+			PreviewSettings = FModulationSettingsProxy();
+		}
 	}
 
 #if !UE_BUILD_SHIPPING
