Trivial: Fix bug with modulation debugger not displaying after entering/exiting PIE (and just utilize last set filters if any)

diff --git a/Engine/Plugins/Runtime/AudioModulation/Source/AudioModulation/Private/AudioModulationDebugger.cpp b/Engine/Plugins/Runtime/AudioModulation/Source/AudioModulation/Private/AudioModulationDebugger.cpp
--- a/Engine/Plugins/Runtime/AudioModulation/Source/AudioModulation/Private/AudioModulationDebugger.cpp
+++ b/Engine/Plugins/Runtime/AudioModulation/Source/AudioModulation/Private/AudioModulationDebugger.cpp
@@ -331,6 +331,13 @@ namespace AudioModulation
 	{
 		check(IsInGameThread());
 
+		// Render stats can get called when toggle did not update bActive, so force active
+		// if called and update toggle state accordingly, utilizing the last set filters.
+		if (!bActive)
+		{
+			bActive = true;
+		}
+
 		if (bShowRenderStatMix)
 		{
 			Y = RenderStatMixMatrix(FilteredMixes, FilteredBuses, Canvas, X, Y, Font);
