Fix non-editor build
#rnx

diff --git a/Engine/Plugins/Runtime/AudioModulation/Source/AudioModulation/Private/SoundControlBus.cpp b/Engine/Plugins/Runtime/AudioModulation/Source/AudioModulation/Private/SoundControlBus.cpp
--- a/Engine/Plugins/Runtime/AudioModulation/Source/AudioModulation/Private/SoundControlBus.cpp
+++ b/Engine/Plugins/Runtime/AudioModulation/Source/AudioModulation/Private/SoundControlBus.cpp
@@ -12,7 +12,9 @@
 
 USoundControlBusBase::USoundControlBusBase(const FObjectInitializer& ObjectInitializer)
 	: Super(ObjectInitializer)
+#if WITH_EDITORONLY_DATA
 	, bOverrideAddress(false)
+#endif
 	, DefaultValue(1.0f)
 	, Min(0.0f)
 	, Max(1.0f)
