- Fix compiling with null modulation implementation
- Fix strict build

diff --git a/Engine/Plugins/Runtime/AudioModulation/Source/AudioModulation/Private/AudioModulationInternal.h b/Engine/Plugins/Runtime/AudioModulation/Source/AudioModulation/Private/AudioModulationInternal.h
--- a/Engine/Plugins/Runtime/AudioModulation/Source/AudioModulation/Private/AudioModulationInternal.h
+++ b/Engine/Plugins/Runtime/AudioModulation/Source/AudioModulation/Private/AudioModulationInternal.h
@@ -10,6 +10,7 @@
 #include "SoundControlBus.h"
 #include "SoundControlBusMix.h"
 #include "SoundModulatorLFO.h"
+#include "IAudioExtensionPlugin.h"
 
 
 #if WITH_AUDIOMODULATION
@@ -103,9 +104,9 @@ namespace AudioModulation
 		void OnExitPIE(bool bSimulating) { }
 #endif // WITH_EDITOR
 
-		void OnInitSound(ISoundModulatable* InSound, const USoundModulationPluginSourceSettingsBase& Settings) { }
+		void OnInitSound(ISoundModulatable& InSound, const USoundModulationPluginSourceSettingsBase& Settings) { }
 		void OnInitSource(const uint32 InSourceId, const FName& AudioComponentUserId, const uint32 NumChannels, const USoundModulationPluginSourceSettingsBase& Settings) { }
-		void OnReleaseSound(ISoundModulatable* InSound) { }
+		void OnReleaseSound(ISoundModulatable& InSound) { }
 		void OnReleaseSource(const uint32 InSourceId) { }
 
 #if !UE_BUILD_SHIPPING
