- Add ability to have shared curves for modulation transforms
- Fix scalar bug in exp inverse expression
- Expose transorm to public API for ModulationEditor updates (in subsequent CL)

diff --git a/Engine/Plugins/Runtime/AudioModulation/Source/AudioModulation/Private/SoundModulationTransform.cpp b/Engine/Plugins/Runtime/AudioModulation/Source/AudioModulation/Private/SoundModulationTransform.cpp
--- a/Engine/Plugins/Runtime/AudioModulation/Source/AudioModulation/Private/SoundModulationTransform.cpp
+++ b/Engine/Plugins/Runtime/AudioModulation/Source/AudioModulation/Private/SoundModulationTransform.cpp
@@ -44,21 +44,27 @@ void FSoundModulationOutputTransform::EvaluateCurve(float& Value) const
 	// If custom curve, evaluate curve and return before calculating alpha
 	if (Curve == ESoundModulatorOutputCurve::Custom)
 	{
-		if (const FRichCurve* RichCurve = FloatCurve.GetRichCurveConst())
+		Value = CurveCustom.Eval(Value);
+		return;
+	}
+
+	// If shared curve, evaluate curve and return before calculating alpha
+	if (Curve == ESoundModulatorOutputCurve::Shared)
+	{
+		if (CurveShared)
 		{
-			Value = RichCurve->Eval(Value);
+			Value = CurveShared->FloatCurve.Eval(Value);
 		}
 		return;
 	}
 
-	// Avoid divide-by-zero & return max if input range is limited to single value
-	if (InputMax == InputMin)
+	// Avoid divide-by-zero & return untransformed if limited to single value
+	if (InputMax >= InputMin)
 	{
-		Value = OutputMax;
 		return;
 	}
 	
-	const float Alpha = FMath::Clamp((Value - InputMin) / (InputMax - InputMin), 0.0f, 1.0f);
+	const float Alpha = (Value - InputMin) / (InputMax - InputMin);
 	switch (Curve)
 	{
 		case ESoundModulatorOutputCurve::Linear:
@@ -79,7 +85,7 @@ void FSoundModulationOutputTransform::EvaluateCurve(float& Value) const
 		{
 			// Alpha is limited to between 0.0f and 1.0f and ExponentialScalar
 			// between 0 and 10 to keep values "sane" and avoid float boundary.
-			Value = ((Alpha - 1.0f) * (Scalar * FMath::Pow(10.0f, Alpha))) + 1.0f;
+			Value = ((Alpha - 1.0f) * FMath::Pow(10.0f, -1.0f * Scalar * Alpha)) + 1.0f;
 		}
 		break;
 
@@ -103,7 +109,7 @@ void FSoundModulationOutputTransform::EvaluateCurve(float& Value) const
 
 		default:
 		{
-			static_assert(static_cast<int32>(ESoundModulatorOutputCurve::Count) == 7, "Possible missing case coverage for output curve.");
+			static_assert(static_cast<int32>(ESoundModulatorOutputCurve::Count) == 8, "Possible missing case coverage for output curve.");
 		}
 		break;
 	}
diff --git a/Engine/Plugins/Runtime/AudioModulation/Source/AudioModulation/Private/SoundModulationTransform.h b/Engine/Plugins/Runtime/AudioModulation/Source/AudioModulation/Private/SoundModulationTransform.h
deleted file mode 100644
diff --git a/Engine/Plugins/Runtime/AudioModulation/Source/AudioModulation/Public/SoundModulationTransform.h b/Engine/Plugins/Runtime/AudioModulation/Source/AudioModulation/Public/SoundModulationTransform.h
new file mode 100644
--- /dev/null
+++ b/Engine/Plugins/Runtime/AudioModulation/Source/AudioModulation/Public/SoundModulationTransform.h
@@ -0,0 +1,107 @@
+// Copyright 1998-2019 Epic Games, Inc. All Rights Reserved.
+#pragma once
+
+#include "CoreMinimal.h"
+#include "Curves/CurveFloat.h"
+#include "UObject/ObjectMacros.h"
+#include "UObject/Object.h"
+
+#include "SoundModulationTransform.generated.h"
+
+
+USTRUCT(BlueprintType)
+struct FSoundModulationInputTransform
+{
+	GENERATED_USTRUCT_BODY()
+
+	FSoundModulationInputTransform();
+
+	/** Minimum value to clamp the input to. */
+	UPROPERTY(EditAnywhere, Category = Modulation, BlueprintReadWrite, meta = (DisplayName = "Input Min", UIMin = "0", UIMax = "1"))
+	float InputMin;
+
+	/** Maximum value to clamp the input to. */
+	UPROPERTY(EditAnywhere, Category = Modulation, BlueprintReadWrite, meta = (DisplayName = "Input Max", UIMin = "0", UIMax = "1"))
+	float InputMax;
+
+	/** Minimum value to clamp the output to. */
+	UPROPERTY(EditAnywhere, Category = Modulation, BlueprintReadWrite, meta = (DisplayName = "Output Min", UIMin = "0", UIMax = "1"))
+	float OutputMin;
+
+	/** Maximum value to clamp the output to. */
+	UPROPERTY(EditAnywhere, Category = Modulation, BlueprintReadWrite, meta = (DisplayName = "Output Max", UIMin = "0", UIMax = "1"))
+	float OutputMax;
+
+	/** Applies transform to provided value */
+	void Apply(float& OutValue) const;
+};
+
+UENUM()
+enum class ESoundModulatorOutputCurve
+{
+	// Expressions
+	Linear		UMETA(DisplayName = "Linear"),
+	Exp			UMETA(DisplayName = "Exponential"),
+	Exp_Inverse UMETA(DisplayName = "Exponential (Inverse)"),
+	Log			UMETA(DisplayName = "Log"),
+	Sin			UMETA(DisplayName = "Sin (Quarter)"),
+	SCurve		UMETA(DisplayName = "Sin (S-Curve)"),
+
+	// Reference a shared curve asset
+	Shared		UMETA(DisplayName = "Shared"),
+
+	// Design a custom curve unique to the owning transform
+	Custom		UMETA(DisplayName = "Custom"),
+
+	Count		UMETA(Hidden),
+};
+
+USTRUCT(BlueprintType)
+struct AUDIOMODULATION_API FSoundModulationOutputTransform
+{
+	GENERATED_USTRUCT_BODY()
+
+	FSoundModulationOutputTransform();
+
+	/** Minimum value to clamp the input to. */
+	UPROPERTY(EditAnywhere, Category = Input, BlueprintReadWrite, meta = (DisplayName = "Input Min", UIMin = "0", UIMax = "1"))
+	float InputMin;
+
+	/** Maximum value to clamp the input to. */
+	UPROPERTY(EditAnywhere, Category = Input, BlueprintReadWrite, meta = (DisplayName = "Input Max", UIMin = "0", UIMax = "1"))
+	float InputMax;
+
+	/** The curve to apply when transforming the output. */
+	UPROPERTY(EditAnywhere, Category = Curve, BlueprintReadWrite)
+	ESoundModulatorOutputCurve Curve;
+
+	/** When curve set to log, exponential or exponential inverse, value is factor 'b' in following equations with output 'y' and input 'x':
+	 *  Exponential: y = x * 10^-b(1-x)
+	 *  Exponential (Inverse): y = ((x - 1) * 10^(-bx)) + 1
+	 *  Logarithmic: y = b * log(x) + 1
+	 */
+	UPROPERTY(EditAnywhere, Category = Input, BlueprintReadWrite, meta = (DisplayName = "Exponential Scalar", ClampMin = "0.1", ClampMax = "10.0", UIMin = "0.1", UIMax = "10.0"))
+	float Scalar;
+
+	/** Custom curve to apply if output curve type is set to 'Custom.' */
+	UPROPERTY(EditAnywhere, Category = Curve, BlueprintReadWrite)
+	FRichCurve CurveCustom;
+
+	/** Asset curve reference to apply if output curve type is set to 'Shared.' */
+	UPROPERTY(EditAnywhere, meta = (DisplayName = "Asset"), Category = Curve, BlueprintReadWrite)
+	UCurveFloat* CurveShared;
+
+	/** Minimum value to clamp output to. */
+	UPROPERTY(EditAnywhere, Category = Output, BlueprintReadWrite)
+	float OutputMin;
+
+	/** Maximum value to clamp output to. */
+	UPROPERTY(EditAnywhere, Category = Output, BlueprintReadWrite)
+	float OutputMax;
+
+	/** Applies transform to provided value */
+	void Apply(float& OutValue) const;
+
+private:
+	void EvaluateCurve(float& Value) const;
+};
