Trivial: Remove invalid mark-up

diff --git a/Engine/Plugins/Runtime/AudioModulation/Source/AudioModulation/Public/SoundModulationTransform.h b/Engine/Plugins/Runtime/AudioModulation/Source/AudioModulation/Public/SoundModulationTransform.h
--- a/Engine/Plugins/Runtime/AudioModulation/Source/AudioModulation/Public/SoundModulationTransform.h
+++ b/Engine/Plugins/Runtime/AudioModulation/Source/AudioModulation/Public/SoundModulationTransform.h
@@ -84,7 +84,7 @@ struct AUDIOMODULATION_API FSoundModulationOutputTransform
 	float Scalar;
 
 	/** Custom curve to apply if output curve type is set to 'Custom.' */
-	UPROPERTY(EditAnywhere, Category = Curve, BlueprintReadWrite)
+	UPROPERTY()
 	FRichCurve CurveCustom;
 
 	/** Asset curve reference to apply if output curve type is set to 'Shared.' */
