Trivial: Fix old custom mod transform editor compile error

diff --git a/Engine/Plugins/Runtime/AudioModulation/Source/AudioModulationEditor/Private/Layouts/SoundModulationTransformLayout.cpp b/Engine/Plugins/Runtime/AudioModulation/Source/AudioModulationEditor/Private/Layouts/SoundModulationTransformLayout.cpp
--- a/Engine/Plugins/Runtime/AudioModulation/Source/AudioModulationEditor/Private/Layouts/SoundModulationTransformLayout.cpp
+++ b/Engine/Plugins/Runtime/AudioModulation/Source/AudioModulationEditor/Private/Layouts/SoundModulationTransformLayout.cpp
@@ -65,7 +65,7 @@ void FSoundModulationOutputTransformLayoutCustomization::CustomizeChildren(TShar
 	TSharedRef<IPropertyHandle>CurveHandle = PropertyHandles.FindChecked(GET_MEMBER_NAME_CHECKED(FSoundModulationOutputTransform, Curve)).ToSharedRef();
 	ChildBuilder.AddProperty(CurveHandle);
 
-	TSharedRef<IPropertyHandle> FloatCurveHandle = PropertyHandles.FindChecked(GET_MEMBER_NAME_CHECKED(FSoundModulationOutputTransform, FloatCurve)).ToSharedRef();
+	TSharedRef<IPropertyHandle> FloatCurveHandle = PropertyHandles.FindChecked(GET_MEMBER_NAME_CHECKED(FSoundModulationOutputTransform, CurveCustom)).ToSharedRef();
 	const TArray<FString> CustomFilters = TArray<FString>({ TEXT("Custom") });
 	ChildBuilder.AddProperty(FloatCurveHandle)
 		.EditCondition(TAttribute<bool>::Create([this, CurveHandle, CustomFilters]() { return IsCurveSet(CurveHandle, CustomFilters); }), nullptr)
