- Fix sounds not reviving due to modulation setting source to 0 volume by calculating if can play when FSoundSource not created for associated FWaveInstance
- Fix modulation low-pass defaulting to incorrect value
- Various polish & clean-up

diff --git a/Engine/Plugins/Runtime/AudioModulation/Source/AudioModulation/Private/AudioModulation.cpp b/Engine/Plugins/Runtime/AudioModulation/Source/AudioModulation/Private/AudioModulation.cpp
--- a/Engine/Plugins/Runtime/AudioModulation/Source/AudioModulation/Private/AudioModulation.cpp
+++ b/Engine/Plugins/Runtime/AudioModulation/Source/AudioModulation/Private/AudioModulation.cpp
@@ -26,6 +26,11 @@ namespace AudioModulation
 		return Impl.Get();
 	}
 
+	float FAudioModulation::CalculateInitialVolume(const USoundModulationPluginSourceSettingsBase& Settings) const
+	{
+		return Impl->CalculateInitialVolume(Settings);
+	}
+
 	void FAudioModulation::Initialize(const FAudioPluginInitializationParams& InitializationParams)
 	{
 		Impl->Initialize(InitializationParams);
diff --git a/Engine/Plugins/Runtime/AudioModulation/Source/AudioModulation/Private/AudioModulationImpl.cpp b/Engine/Plugins/Runtime/AudioModulation/Source/AudioModulation/Private/AudioModulationImpl.cpp
--- a/Engine/Plugins/Runtime/AudioModulation/Source/AudioModulation/Private/AudioModulationImpl.cpp
+++ b/Engine/Plugins/Runtime/AudioModulation/Source/AudioModulation/Private/AudioModulationImpl.cpp
@@ -87,7 +87,9 @@ namespace
 namespace AudioModulation
 {
 	FAudioModulationImpl::FAudioModulationImpl()
+#if WITH_EDITOR
 		: PreviewSound(nullptr)
+#endif // WITH_EDITOR
 	{
 	}
 
@@ -103,7 +105,7 @@ namespace AudioModulation
 		// This informs user that modifying sound's settings does not translate to
 		// currently playing sound.
 		const uint32 SettingsId = InSettings.GetUniqueID();
-		auto UpdateSettings = [this, SettingsId]()
+		RunCommandOnAudioThread([this, SettingsId]()
 		{
 			if (PreviewSound)
 			{
@@ -157,9 +159,7 @@ namespace AudioModulation
 					}
 				}
 			}
-		};
-
-		IsInAudioThread() ? UpdateSettings() : FAudioThread::RunCommandOnAudioThread(UpdateSettings);
+		});
 	}
 #endif // WITH_EDITOR
 
@@ -168,14 +168,16 @@ namespace AudioModulation
 		check(IsInAudioThread());
 
 		const USoundModulationSettings* Settings = CastChecked<USoundModulationSettings>(&InSettings);
-		const uint32 SoundId = InSound.GetObjectId();
 
+#if WITH_EDITOR
 		if (InSound.IsPreviewSound())
 		{
 			PreviewSound = &InSound;
 			PreviewSettings = FModulationSettingsProxy(*Settings);
 		}
+#endif // WITH_EDITOR
 
+		const uint32 SoundId = InSound.GetObjectId();
 		if (!SoundSettings.Contains(SoundId))
 		{
 			SoundSettings.Add(SoundId, FModulationSettingsProxy(*Settings));
@@ -257,38 +259,42 @@ namespace AudioModulation
 		check(IsInAudioThread());
 		check(InSound.GetObjectId() != INDEX_NONE);
 
-		const FModulationSettingsProxy& Settings = SoundSettings.FindChecked(InSound.GetObjectId());
-		for (FBusMixId MixId : Settings.Mixes)
+		// Settings can be null if sound settings were modified via the editor while auditioning or in PIE
+		if (const FModulationSettingsProxy* Settings = SoundSettings.Find(InSound.GetObjectId()))
 		{
-			DeactivateBusMix(MixId, &InSound);
-		}
+			for (FBusMixId MixId : Settings->Mixes)
+			{
+				DeactivateBusMix(MixId, &InSound);
+			}
 
-		for (const FModulationInputProxy& Input : Settings.Volume.InputProxies)
-		{
-			DeactivateBus(Input.BusId, &InSound);
-		}
+			for (const FModulationInputProxy& Input : Settings->Volume.InputProxies)
+			{
+				DeactivateBus(Input.BusId, &InSound);
+			}
 
-		for (const FModulationInputProxy& Input : Settings.Pitch.InputProxies)
-		{
-			DeactivateBus(Input.BusId, &InSound);
-		}
+			for (const FModulationInputProxy& Input : Settings->Pitch.InputProxies)
+			{
+				DeactivateBus(Input.BusId, &InSound);
+			}
 
-		for (const FModulationInputProxy& Input : Settings.Lowpass.InputProxies)
-		{
-			DeactivateBus(Input.BusId, &InSound);
-		}
+			for (const FModulationInputProxy& Input : Settings->Lowpass.InputProxies)
+			{
+				DeactivateBus(Input.BusId, &InSound);
+			}
 
-		for (const FModulationInputProxy& Input : Settings.Highpass.InputProxies)
-		{
-			DeactivateBus(Input.BusId, &InSound);
+			for (const FModulationInputProxy& Input : Settings->Highpass.InputProxies)
+			{
+				DeactivateBus(Input.BusId, &InSound);
+			}
 		}
 
-		if (InSound.IsPreviewSound())
+#if WITH_EDITOR
+		if (InSound.IsPreviewSound() && &InSound == PreviewSound)
 		{
-			check(&InSound == PreviewSound);
 			PreviewSound = nullptr;
 			PreviewSettings = FModulationSettingsProxy();
 		}
+#endif // WITH_EDITOR
 	}
 
 #if !UE_BUILD_SHIPPING
@@ -455,9 +461,19 @@ namespace AudioModulation
 		return OutValue;
 	}
 
+	float FAudioModulationImpl::CalculateInitialVolume(const USoundModulationPluginSourceSettingsBase& SettingsBase) const
+	{
+		check(IsInAudioThread());
+
+		const USoundModulationSettings* Settings = CastChecked<USoundModulationSettings>(&SettingsBase);
+		FModulationPatchProxy VolumePatch(Settings->Volume);
+
+		return CalculateModulationValue(VolumePatch);
+	}
+
 	void FAudioModulationImpl::DeactivateBusMix(const FBusMixId InBusMixId, const ISoundModulatable* InSound)
 	{
-		auto DeactivateMix = [this, InBusMixId, InSound]()
+		RunCommandOnAudioThread([this, InBusMixId, InSound]()
 		{
 			check(IsInAudioThread());
 
@@ -478,9 +494,7 @@ namespace AudioModulation
 					}
 				}
 			}
-		};
-
-		IsInAudioThread() ? DeactivateMix() : FAudioThread::RunCommandOnAudioThread(DeactivateMix);
+		});
 	}
 
 	void FAudioModulationImpl::DeactivateBus(const FBusId InBusId, const ISoundModulatable* InSound)
@@ -596,7 +610,7 @@ namespace AudioModulation
 			Pair.Value.Update(Elapsed, ActiveBuses);
 			if (Pair.Value.CanDestroy())
 			{
-				UE_LOG(LogAudioModulation, Log, TEXT("Audio modulation mix '%u' stopped."), static_cast<uint32>(Pair.Key));
+				UE_LOG(LogAudioModulation, Log, TEXT("Audio modulation mix '%s' stopped."), *Pair.Value.GetName());
 				MixesToDeactivate.Add(Pair.Key);
 			}
 		}
diff --git a/Engine/Plugins/Runtime/AudioModulation/Source/AudioModulation/Private/AudioModulationInternal.h b/Engine/Plugins/Runtime/AudioModulation/Source/AudioModulation/Private/AudioModulationInternal.h
--- a/Engine/Plugins/Runtime/AudioModulation/Source/AudioModulation/Private/AudioModulationInternal.h
+++ b/Engine/Plugins/Runtime/AudioModulation/Source/AudioModulation/Private/AudioModulationInternal.h
@@ -30,7 +30,6 @@ namespace AudioModulation
 
 #if WITH_EDITOR
 		void OnEditPluginSettings(const USoundModulationPluginSourceSettingsBase& Settings);
-		void OnExitPIE(bool bSimulating);
 #endif // WITH_EDITOR
 
 		void OnInitSound(ISoundModulatable& InSound, const USoundModulationPluginSourceSettingsBase& Settings);
@@ -42,9 +41,11 @@ namespace AudioModulation
 		void ActivateBusMix(const USoundControlBusMix& InBusMix, const ISoundModulatable* Sound = nullptr);
 		void ActivateLFO(const USoundBusModulatorLFO& InLFO, const ISoundModulatable* Sound = nullptr);
 
+		float CalculateInitialVolume(const USoundModulationPluginSourceSettingsBase& InSettingsBase) const;
+
 		/**
-		 * Deactivates respective type if InSoundId is invalid (default behavior).  Only deactivates type if
-		 * provided a InSoundId that is releasing, type is set to AutoActivate, and sounds are no longer referencing type.
+		 * Deactivates respectively typed (i.e. BusMix, Bus, etc.) object proxy if either Sound is invalid (default behavior),
+		 * or if provided Sound, only deactivates if object is set to AutoActivate and sounds are no longer referencing provided object.
 		 */
 		void DeactivateBusMix(const FBusMixId InBusMixId, const ISoundModulatable* Sound = nullptr);
 		void DeactivateBus(const FBusId InBusId, const ISoundModulatable* Sound = nullptr);
@@ -53,7 +54,6 @@ namespace AudioModulation
 		bool IsBusActive(const FBusId InBusId) const;
 		bool IsLFOActive(const FLFOId InLFOId) const;
 
-		void ProcessAudio(const FAudioPluginSourceInputData& InputData, FAudioPluginSourceOutputData& OutputData);
 		void ProcessControls(const uint32 InSourceId, FSoundModulationControls& OutControls);
 		void ProcessModulators(const float Elapsed);
 
@@ -83,9 +83,13 @@ namespace AudioModulation
 
 	private:
 		FAudioModulationDebugger Debugger;
+#endif // !UE_BUILD_SHIPPING
+
+#if WITH_EDITOR
+	private:
 		ISoundModulatable* PreviewSound;
 		FModulationSettingsProxy PreviewSettings;
-#endif // !UE_BUILD_SHIPPING
+#endif // WITH_EDITOR
 	};
 } // namespace AudioModulation
 
@@ -101,7 +105,6 @@ namespace AudioModulation
 
 #if WITH_EDITOR
 		void OnEditPluginSettings(const USoundModulationPluginSourceSettingsBase& Settings) { }
-		void OnExitPIE(bool bSimulating) { }
 #endif // WITH_EDITOR
 
 		void OnInitSound(ISoundModulatable& InSound, const USoundModulationPluginSourceSettingsBase& Settings) { }
@@ -118,11 +121,13 @@ namespace AudioModulation
 		void ActivateBus(const USoundControlBusBase& InBus, const ISoundModulatable* Sound = nullptr) { }
 		void ActivateBusMix(const USoundControlBusMix& InBusMix, const ISoundModulatable* Sound = nullptr) { }
 		void ActivateLFO(const USoundBusModulatorLFO& InLFO, const ISoundModulatable* Sound = nullptr) { }
+
+		float CalculateInitialVolume(const USoundModulationPluginSourceSettingsBase& InSettingsBase) const { return 1.0f; }
+
 		void DeactivateBusMix(const FBusMixId BusMixId, const ISoundModulatable* Sound = nullptr) { }
 		void DeactivateBus(const FBusId BusId, const ISoundModulatable* Sound = nullptr) { }
 		void DeactivateLFO(const FLFOId LFOId, const ISoundModulatable* Sound = nullptr) { }
 
-		void ProcessAudio(const FAudioPluginSourceInputData& InputData, FAudioPluginSourceOutputData& OutputData) { }
 		void ProcessControls(const uint32 InSourceId, FSoundModulationControls& OutControls) { }
 		void ProcessModulators(const float Elapsed) { }
 
diff --git a/Engine/Plugins/Runtime/AudioModulation/Source/AudioModulation/Private/SoundModulationPatch.cpp b/Engine/Plugins/Runtime/AudioModulation/Source/AudioModulation/Private/SoundModulationPatch.cpp
--- a/Engine/Plugins/Runtime/AudioModulation/Source/AudioModulation/Private/SoundModulationPatch.cpp
+++ b/Engine/Plugins/Runtime/AudioModulation/Source/AudioModulation/Private/SoundModulationPatch.cpp
@@ -37,21 +37,46 @@ void USoundModulationSettings::OnPostEditChange(UWorld* World)
 
 void USoundModulationSettings::PostEditChangeProperty(FPropertyChangedEvent& PropertyChangedEvent)
 {
-	if (PropertyChangedEvent.Property)
+	OnPostEditChange(GetWorld());
+}
+
+void USoundModulationSettings::PostEditChangeChainProperty(FPropertyChangedChainEvent& PropertyChangedEvent)
+{
+	OnPostEditChange(GetWorld());
+}
+#endif // WITH_EDITOR
+
+void FSoundVolumeModulationPatch::GenerateProxies(TArray<AudioModulation::FModulationInputProxy>& InputProxies) const
+{
+	for (const FSoundVolumeModulationInput& Input : Inputs)
 	{
-		OnPostEditChange(GetWorld());
+		InputProxies.Emplace_GetRef(Input);
 	}
 }
 
-void USoundModulationSettings::PostEditChangeChainProperty(FPropertyChangedChainEvent& PropertyChangedEvent)
+void FSoundPitchModulationPatch::GenerateProxies(TArray<AudioModulation::FModulationInputProxy>& InputProxies) const
 {
-	if (PropertyChangedEvent.Property)
+	for (const FSoundPitchModulationInput& Input : Inputs)
 	{
-		OnPostEditChange(GetWorld());
+		InputProxies.Emplace_GetRef(Input);
 	}
+}
 
+void FSoundHPFModulationPatch::GenerateProxies(TArray<AudioModulation::FModulationInputProxy>& InputProxies) const
+{
+	for (const FSoundHPFModulationInput& Input : Inputs)
+	{
+		InputProxies.Emplace_GetRef(Input);
+	}
+}
+
+void FSoundLPFModulationPatch::GenerateProxies(TArray<AudioModulation::FModulationInputProxy>& InputProxies) const
+{
+	for (const FSoundLPFModulationInput& Input : Inputs)
+	{
+		InputProxies.Emplace_GetRef(Input);
+	}
 }
-#endif // WITH_EDITOR
 
 FSoundModulationOutput::FSoundModulationOutput()
 	: Operator(ESoundModulatorOperator::Multiply)
diff --git a/Engine/Plugins/Runtime/AudioModulation/Source/AudioModulation/Private/SoundModulationPatch.h b/Engine/Plugins/Runtime/AudioModulation/Source/AudioModulation/Private/SoundModulationPatch.h
deleted file mode 100644
diff --git a/Engine/Plugins/Runtime/AudioModulation/Source/AudioModulation/Private/SoundModulationTransform.cpp b/Engine/Plugins/Runtime/AudioModulation/Source/AudioModulation/Private/SoundModulationTransform.cpp
--- a/Engine/Plugins/Runtime/AudioModulation/Source/AudioModulation/Private/SoundModulationTransform.cpp
+++ b/Engine/Plugins/Runtime/AudioModulation/Source/AudioModulation/Private/SoundModulationTransform.cpp
@@ -59,7 +59,7 @@ void FSoundModulationOutputTransform::EvaluateCurve(float& Value) const
 	}
 
 	// Avoid divide-by-zero & return untransformed if limited to single value
-	if (InputMax >= InputMin)
+	if (InputMax <= InputMin)
 	{
 		return;
 	}
diff --git a/Engine/Plugins/Runtime/AudioModulation/Source/AudioModulation/Public/AudioModulation.h b/Engine/Plugins/Runtime/AudioModulation/Source/AudioModulation/Public/AudioModulation.h
--- a/Engine/Plugins/Runtime/AudioModulation/Source/AudioModulation/Public/AudioModulation.h
+++ b/Engine/Plugins/Runtime/AudioModulation/Source/AudioModulation/Public/AudioModulation.h
@@ -27,8 +27,9 @@ namespace AudioModulation
 		virtual ~FAudioModulation() = default;
 
 		//~ Begin IAudioModulation implementation
-		virtual void Initialize(const FAudioPluginInitializationParams& InitializationParams) override;
+		virtual float CalculateInitialVolume(const USoundModulationPluginSourceSettingsBase& InSettingsBase) const override;
 
+		virtual void Initialize(const FAudioPluginInitializationParams& InitializationParams) override;
 		virtual void OnInitSound(ISoundModulatable& Sound, const USoundModulationPluginSourceSettingsBase& Settings) override;
 		virtual void OnInitSource(const uint32 SourceId, const FName& AudioComponentUserId, const uint32 NumChannels, const USoundModulationPluginSourceSettingsBase& Settings) override;
 
diff --git a/Engine/Plugins/Runtime/AudioModulation/Source/AudioModulation/Public/SoundModulationPatch.h b/Engine/Plugins/Runtime/AudioModulation/Source/AudioModulation/Public/SoundModulationPatch.h
new file mode 100644
--- /dev/null
+++ b/Engine/Plugins/Runtime/AudioModulation/Source/AudioModulation/Public/SoundModulationPatch.h
@@ -0,0 +1,275 @@
+// Copyright 1998-2019 Epic Games, Inc. All Rights Reserved.
+#pragma once
+
+#include "CoreMinimal.h"
+#include "UObject/ObjectMacros.h"
+#include "UObject/Object.h"
+#include "SoundModulationTransform.h"
+#include "SoundModulationValue.h"
+#include "SoundModulatorBase.h"
+#include "SoundControlBus.h"
+#include "SoundControlBusMix.h"
+
+#include "SoundModulationPatch.generated.h"
+
+
+// Forward Declarations
+namespace AudioModulation
+{
+	struct FModulationInputProxy;
+} // namespace AudioModulation
+
+
+USTRUCT(BlueprintType)
+struct FSoundModulationOutput
+{
+	GENERATED_USTRUCT_BODY()
+
+	FSoundModulationOutput();
+
+	/** Operator used when mixing input values */
+	UPROPERTY(VisibleAnywhere, BlueprintReadOnly, Category = Modulation)
+	ESoundModulatorOperator Operator;
+
+	/** Final transform before passing to output */
+	UPROPERTY(EditAnywhere, Category = Modulation, BlueprintReadWrite, meta = (ShowOnlyInnerProperties))
+	FSoundModulationOutputTransform Transform;
+};
+
+USTRUCT(BlueprintType)
+struct FSoundModulationInputBase
+{
+	GENERATED_USTRUCT_BODY()
+
+	FSoundModulationInputBase();
+	virtual ~FSoundModulationInputBase() = default;
+
+	/** Get the modulated input value on parent patch initialization and hold that value for its lifetime */
+	UPROPERTY(EditAnywhere, BlueprintReadWrite, Category = Input, meta = (DisplayName = "Sample-And-Hold"))
+	uint8 bSampleAndHold : 1;
+
+	/** Transform to apply to the input prior to mix phase */
+	UPROPERTY(EditAnywhere, BlueprintReadWrite, Category = Input, meta = (ShowOnlyInnerProperties))
+	FSoundModulationInputTransform Transform;
+
+	virtual const USoundControlBusBase* GetBus() const PURE_VIRTUAL(FSoundModulationInputBase::GetBus, return nullptr; );
+};
+
+USTRUCT(BlueprintType)
+struct FSoundVolumeModulationInput : public FSoundModulationInputBase
+{
+	GENERATED_USTRUCT_BODY()
+
+	FSoundVolumeModulationInput();
+
+	/** The input bus */
+	UPROPERTY(EditAnywhere, BlueprintReadWrite, Category = Input)
+	USoundVolumeControlBus* Bus;
+
+	virtual const USoundControlBusBase* GetBus() const { return Cast<USoundControlBusBase>(Bus); }
+};
+
+USTRUCT(BlueprintType)
+struct FSoundPitchModulationInput : public FSoundModulationInputBase
+{
+	GENERATED_USTRUCT_BODY()
+
+	FSoundPitchModulationInput();
+
+	/** The input bus */
+	UPROPERTY(EditAnywhere, BlueprintReadWrite, Category = Input)
+	USoundPitchControlBus* Bus;
+
+	virtual const USoundControlBusBase* GetBus() const { return Cast<USoundControlBusBase>(Bus); }
+};
+
+USTRUCT(BlueprintType)
+struct FSoundLPFModulationInput : public FSoundModulationInputBase
+{
+	GENERATED_USTRUCT_BODY()
+
+	FSoundLPFModulationInput();
+
+	/** The input bus */
+	UPROPERTY(EditAnywhere, BlueprintReadWrite, Category = Input)
+	USoundLPFControlBus* Bus;
+
+	virtual const USoundControlBusBase* GetBus() const { return Cast<USoundControlBusBase>(Bus); }
+};
+
+USTRUCT(BlueprintType)
+struct FSoundHPFModulationInput : public FSoundModulationInputBase
+{
+	GENERATED_USTRUCT_BODY()
+
+	FSoundHPFModulationInput();
+
+	/** The input bus */
+	UPROPERTY(EditAnywhere, BlueprintReadWrite, Category = Input)
+	USoundHPFControlBus* Bus;
+
+	virtual const USoundControlBusBase* GetBus() const { return Cast<USoundControlBusBase>(Bus); }
+};
+
+USTRUCT(BlueprintType)
+struct FSoundModulationPatchBase
+{
+	GENERATED_USTRUCT_BODY()
+
+	FSoundModulationPatchBase();
+
+	virtual ~FSoundModulationPatchBase() = default;
+
+	/** Default value of patch (value used when all inputs are either not provided or not active) */
+	UPROPERTY(EditAnywhere, BlueprintReadWrite, Category = Inputs, meta = (UIMin = "0", UIMax = "1"))
+	float DefaultInputValue;
+
+	/** Final modulation parameters to apply */
+	UPROPERTY(EditAnywhere, BlueprintReadWrite, Category = Output, meta = (ShowOnlyInnerProperties))
+	FSoundModulationOutput Output;
+
+	virtual void GenerateProxies(TArray<AudioModulation::FModulationInputProxy>& InputProxies) const PURE_VIRTUAL(FSoundModulationPatchBase::GenerateProxies, );
+};
+
+USTRUCT(BlueprintType)
+struct FSoundVolumeModulationPatch : public FSoundModulationPatchBase
+{
+	GENERATED_USTRUCT_BODY()
+
+	/** Modulation inputs */
+	UPROPERTY(EditAnywhere, BlueprintReadWrite, Category = Inputs)
+	TArray<FSoundVolumeModulationInput> Inputs;
+
+	virtual void GenerateProxies(TArray<AudioModulation::FModulationInputProxy>& InputProxies) const override;
+};
+
+USTRUCT(BlueprintType)
+struct FSoundPitchModulationPatch : public FSoundModulationPatchBase
+{
+	GENERATED_USTRUCT_BODY()
+
+	/** Modulation inputs */
+	UPROPERTY(EditAnywhere, BlueprintReadWrite, Category = Inputs)
+	TArray<FSoundPitchModulationInput> Inputs;
+
+	virtual void GenerateProxies(TArray<AudioModulation::FModulationInputProxy>& InputProxies) const override;
+};
+
+USTRUCT(BlueprintType)
+struct FSoundLPFModulationPatch : public FSoundModulationPatchBase
+{
+	GENERATED_USTRUCT_BODY()
+
+	/** Modulation inputs */
+	UPROPERTY(EditAnywhere, BlueprintReadWrite, Category = Inputs)
+	TArray<FSoundLPFModulationInput> Inputs;
+
+	virtual void GenerateProxies(TArray<AudioModulation::FModulationInputProxy>& InputProxies) const override;
+};
+
+USTRUCT(BlueprintType)
+struct FSoundHPFModulationPatch : public FSoundModulationPatchBase
+{
+	GENERATED_USTRUCT_BODY()
+
+	/** Modulation inputs */
+	UPROPERTY(EditAnywhere, BlueprintReadWrite, Category = Inputs)
+	TArray<FSoundHPFModulationInput> Inputs;
+
+	virtual void GenerateProxies(TArray<AudioModulation::FModulationInputProxy>& InputProxies) const override;
+};
+
+UCLASS(config = Engine, editinlinenew, BlueprintType, MinimalAPI, autoExpandCategories = (Volume, Pitch, Highpass, Lowpass))
+class USoundModulationSettings : public USoundModulationPluginSourceSettingsBase
+{
+	GENERATED_UCLASS_BODY()
+
+public:
+	UPROPERTY(EditAnywhere, BlueprintReadWrite, Category = Volume)
+	FSoundVolumeModulationPatch Volume;
+
+	UPROPERTY(EditAnywhere, BlueprintReadWrite, Category = Pitch)
+	FSoundPitchModulationPatch Pitch;
+
+	UPROPERTY(EditAnywhere, BlueprintReadWrite, Category = Highpass)
+	FSoundHPFModulationPatch Highpass;
+
+	UPROPERTY(EditAnywhere, BlueprintReadWrite, Category = Lowpass)
+	FSoundLPFModulationPatch Lowpass;
+
+	// Mixes that will applied and removed when sounds utilizing settings
+	// play and stop respectively. If mix has already been applied manually,
+	// mix will be removed once all sound settings referencing mix stop. Manual
+	// mix activation is ignored if already activated by means of modulation settings.
+	UPROPERTY(EditAnywhere, BlueprintReadWrite, Category = Mixes)
+	TArray<USoundControlBusMix*> Mixes;
+
+#if WITH_EDITOR
+	AUDIOMODULATION_API void OnPostEditChange(UWorld* World);
+
+	virtual void PostEditChangeProperty(FPropertyChangedEvent& PropertyChangedEvent) override;
+	virtual void PostEditChangeChainProperty(FPropertyChangedChainEvent& PropertyChangedEvent) override;
+#endif
+};
+
+
+namespace AudioModulation
+{
+	/** Modulation input instance */
+	struct FModulationInputProxy
+	{
+		FModulationInputProxy();
+		FModulationInputProxy(const FSoundModulationInputBase& Patch);
+
+		FBusId BusId;
+		FSoundModulationInputTransform Transform;
+		uint8 bSampleAndHold : 1;
+	};
+
+	/** Patch applied as the final stage of a modulation chain prior to output on the sound level (Always active, never removed) */
+	struct FModulationOutputProxy
+	{
+		FModulationOutputProxy();
+		FModulationOutputProxy(const FSoundModulationOutput& Patch);
+
+		/** Whether patch has been initialized or not */
+		uint8 bInitialized : 1;
+
+		/** Operator used to calculate the output proxy value */
+		ESoundModulatorOperator Operator;
+
+		/** Cached value of sample-and-hold input values */
+		float SampleAndHoldValue;
+
+		/** Final transform before passing to output */
+		FSoundModulationOutputTransform Transform;
+	};
+
+	struct FModulationPatchProxy
+	{
+		FModulationPatchProxy();
+		FModulationPatchProxy(const FSoundModulationPatchBase& Patch);
+
+		/** Default value of patch (value used when inputs are either not provided or not active)*/
+		float DefaultInputValue;
+
+		/** Optional modulation inputs */
+		TArray<FModulationInputProxy> InputProxies;
+
+		/** Final output modulation post input combination */
+		FModulationOutputProxy OutputProxy;
+	};
+
+	struct FModulationSettingsProxy : TModulatorProxyBase<uint32>
+	{
+		FModulationSettingsProxy();
+		FModulationSettingsProxy(const USoundModulationSettings& Settings);
+
+		FModulationPatchProxy Volume;
+		FModulationPatchProxy Pitch;
+		FModulationPatchProxy Lowpass;
+		FModulationPatchProxy Highpass;
+
+		TArray<FBusMixId> Mixes;
+	};
+} // namespace AudioModulation
diff --git a/Engine/Source/Runtime/Engine/Private/ActiveSound.cpp b/Engine/Source/Runtime/Engine/Private/ActiveSound.cpp
--- a/Engine/Source/Runtime/Engine/Private/ActiveSound.cpp
+++ b/Engine/Source/Runtime/Engine/Private/ActiveSound.cpp
@@ -657,7 +657,7 @@ void FActiveSound::UpdateWaveInstances(TArray<FWaveInstance*> &InWaveInstances,
 			{
 				check(WaveInstance);
 
-				float WaveInstanceVolume = WaveInstance->GetVolumeWithDistanceAttenuation();
+				float WaveInstanceVolume = WaveInstance->GetVolumeWithDistanceAttenuation() * WaveInstance->GetDynamicVolume();
 				if (WaveInstanceVolume > VolumeConcurrency)
 				{
 					VolumeConcurrency = WaveInstanceVolume;
diff --git a/Engine/Source/Runtime/Engine/Private/Audio.cpp b/Engine/Source/Runtime/Engine/Private/Audio.cpp
--- a/Engine/Source/Runtime/Engine/Private/Audio.cpp
+++ b/Engine/Source/Runtime/Engine/Private/Audio.cpp
@@ -87,6 +87,14 @@ FAutoConsoleVariableRef CVarOcclusionFilterScaleEnabled(
 	TEXT("0: Not Enabled, 1: Enabled"),
 	ECVF_Default);
 
+static int32 BypassPlayWhenSilentCVar = 0;
+FAutoConsoleVariableRef CVarBypassPlayWhenSilent(
+	TEXT("au.BypassPlayWhenSilent"),
+	BypassPlayWhenSilentCVar,
+	TEXT("When set to 1, ignores the Play When Silent flag for non-procedural sources.\n")
+	TEXT("0: Honor the Play When Silent flag, 1: stop all silent non-procedural sources."),
+	ECVF_Default);
+
 static int32 AllowReverbForMultichannelSources = 1;
 FAutoConsoleVariableRef CvarAllowReverbForMultichannelSources(
 	TEXT("au.AllowReverbForMultichannelSources"),
@@ -837,6 +845,57 @@ FWaveInstance::FWaveInstance(const UPTRINT InWaveInstanceHash, FActiveSound& InA
 	TypeHash = ++TypeHashCounter;
 }
 
+bool FWaveInstance::IsPlaying() const
+{
+	check(ActiveSound);
+
+	if (!WaveData)
+	{
+		return false;
+	}
+
+	// TODO: move out of audio.  Subtitle system should be separate and just set VirtualizationMode to PlayWhenSilent
+	const bool bHasSubtitles = ActiveSound->bHandleSubtitles && (ActiveSound->bHasExternalSubtitles || WaveData->Subtitles.Num() > 0);
+	if (bHasSubtitles)
+	{
+		return true;
+	}
+
+	if (ActiveSound->IsPlayWhenSilent() && (!BypassPlayWhenSilentCVar || WaveData->bProcedural))
+	{
+		return true;
+	}
+
+	// Modulation volume check must be performed separately from non-modulation volume check if zeroed as this could cause
+	// sources to stop and not be able to be restarted. Because modulation controls are processed on the source level, this
+	// check enables the modulation plugin to determine voice eligibility without having to process all wave instances not
+	// currently sourcing control data.
+	float ModVolume = SoundModulationControls.Volume;
+	if (ModulationPluginSettings)
+	{
+		check(ActiveSound->AudioDevice);
+		FAudioDevice& AudioDevice = *ActiveSound->AudioDevice;
+		if (AudioDevice.IsModulationPluginEnabled())
+		{
+			check(AudioDevice.ModulationInterface);
+			ModVolume = AudioDevice.ModulationInterface->CalculateInitialVolume(*ModulationPluginSettings);
+		}
+	}
+	
+	const float WaveInstanceVolume = ModVolume * Volume * VolumeMultiplier * DistanceAttenuation * GetDynamicVolume();
+	if (WaveInstanceVolume > KINDA_SMALL_NUMBER)
+	{
+		return true;
+	}
+
+	if (ActiveSound->ComponentVolumeFader.IsFadingIn())
+	{
+		return true;
+	}
+
+	return false;
+}
+
 /**
  * Notifies the wave instance that it has finished.
  */
@@ -975,6 +1034,7 @@ float FWaveInstance::GetVolume() const
 
 bool FWaveInstance::ShouldStopDueToMaxConcurrency() const
 {
+	check(ActiveSound);
 	return ActiveSound->bShouldStopDueToMaxConcurrency;
 }
 
diff --git a/Engine/Source/Runtime/Engine/Private/SoundWave.cpp b/Engine/Source/Runtime/Engine/Private/SoundWave.cpp
--- a/Engine/Source/Runtime/Engine/Private/SoundWave.cpp
+++ b/Engine/Source/Runtime/Engine/Private/SoundWave.cpp
@@ -31,14 +31,6 @@
 
 #include "Misc/CommandLine.h"
 
-static int32 BypassPlayWhenSilentCVar = 0;
-FAutoConsoleVariableRef CVarBypassPlayWhenSilent(
-	TEXT("au.BypassPlayWhenSilent"),
-	BypassPlayWhenSilentCVar,
-	TEXT("When set to 1, ignores the Play When Silent flag for non-procedural sources.\n")
-	TEXT("0: Honor the Play When Silent flag, 1: stop all silent non-procedural sources."),
-	ECVF_Default);
-
 static int32 LoadIntoCacheOnPostLoadCVar = 0;
 FAutoConsoleVariableRef CVarLoadIntoCacheOnPostLoad(
 	TEXT("au.streamcache.LoadIntoCacheOnPostLoad"),
@@ -1795,13 +1787,7 @@ void USoundWave::Parse(FAudioDevice* AudioDevice, const UPTRINT NodeWaveInstance
 
 	WaveInstance->bIsAmbisonics = bIsAmbisonics;
 
-	const float WaveInstanceVolume = WaveInstance->GetVolumeWithDistanceAttenuation() * WaveInstance->GetDynamicVolume();
-
-	// When the BypassVirtualizeWhenSilent cvar is enabled, we should only honor bVirtualizeWhenSilent for procedural sounds:
-	const bool bHasSubtitles = ActiveSound.bHandleSubtitles && (ActiveSound.bHasExternalSubtitles || (Subtitles.Num() > 0));
-	const bool bStarted = WaveInstanceVolume > KINDA_SMALL_NUMBER || ActiveSound.ComponentVolumeFader.IsFadingIn();
-	const bool bCanPlayWhenSilent = ActiveSound.IsPlayWhenSilent() && (!BypassPlayWhenSilentCVar || bProcedural);
-	if (bStarted || bCanPlayWhenSilent || bHasSubtitles)
+	if (WaveInstance->IsPlaying())
 	{
 		WaveInstances.Add(WaveInstance);
 		ActiveSound.bFinished = false;
diff --git a/Engine/Source/Runtime/Engine/Public/Audio.h b/Engine/Source/Runtime/Engine/Public/Audio.h
--- a/Engine/Source/Runtime/Engine/Public/Audio.h
+++ b/Engine/Source/Runtime/Engine/Public/Audio.h
@@ -403,6 +403,11 @@ public:
 	void SetStopping(const bool bInIsStopping) { bIsStopping = bInIsStopping; }
 	bool IsStopping() const { return bIsStopping; }
 
+	/** Returns whether or not the WaveInstance is actively playing sound or set to
+	  * play when silent.
+	  */
+	bool IsPlaying() const;
+
 	/** Returns the volume multiplier on the wave instance. */
 	float GetVolumeMultiplier() const { return VolumeMultiplier; }
 
diff --git a/Engine/Source/Runtime/Engine/Public/IAudioExtensionPlugin.h b/Engine/Source/Runtime/Engine/Public/IAudioExtensionPlugin.h
--- a/Engine/Source/Runtime/Engine/Public/IAudioExtensionPlugin.h
+++ b/Engine/Source/Runtime/Engine/Public/IAudioExtensionPlugin.h
@@ -558,6 +558,9 @@ public:
 	/** Virtual destructor */
 	virtual ~IAudioModulation() { }
 
+	/** Calculates initial volume to determine if sound is audible using base settings data */
+	virtual float CalculateInitialVolume(const USoundModulationPluginSourceSettingsBase& InSettingsBase) const { return 1.0f; }
+
 	/** Initialize the modulation plugin with the same rate and number of sources */
 	virtual void Initialize(const FAudioPluginInitializationParams& InitializationParams) { }
 
