- Move out UObject types to public API and proxies to separate private implementation file
- Sanatize bus value by clamping in editor
- Generic modulation control buses/mixing boilerplate (not enabled)
- Disable modulation curve editor to fix crash in editor

diff --git a/Engine/Plugins/Runtime/AudioModulation/Source/AudioModulation/Private/AudioModulation.cpp b/Engine/Plugins/Runtime/AudioModulation/Source/AudioModulation/Private/AudioModulation.cpp
--- a/Engine/Plugins/Runtime/AudioModulation/Source/AudioModulation/Private/AudioModulation.cpp
+++ b/Engine/Plugins/Runtime/AudioModulation/Source/AudioModulation/Private/AudioModulation.cpp
@@ -80,10 +80,10 @@ namespace AudioModulation
 	}
 #endif // !UE_BUILD_SHIPPING
 
-	void FAudioModulation::ProcessControls(const uint32 SourceId, FSoundModulationControls& Controls)
+	bool FAudioModulation::ProcessControls(const uint32 InSourceId, FSoundModulationControls& OutControls)
 	{
 		SCOPE_CYCLE_COUNTER(STAT_AudioModulationProcessControls);
-		Impl->ProcessControls(SourceId, Controls);
+		return Impl->ProcessControls(InSourceId, OutControls);
 	}
 
 	void FAudioModulation::ProcessModulators(const float Elapsed)
diff --git a/Engine/Plugins/Runtime/AudioModulation/Source/AudioModulation/Private/AudioModulationDebugger.h b/Engine/Plugins/Runtime/AudioModulation/Source/AudioModulation/Private/AudioModulationDebugger.h
--- a/Engine/Plugins/Runtime/AudioModulation/Source/AudioModulation/Private/AudioModulationDebugger.h
+++ b/Engine/Plugins/Runtime/AudioModulation/Source/AudioModulation/Private/AudioModulationDebugger.h
@@ -7,6 +7,7 @@
 #if !UE_BUILD_SHIPPING
 #include "SoundControlBus.h"
 #include "SoundControlBusMix.h"
+#include "SoundModulationProxy.h"
 #include "SoundModulatorLFO.h"
 
 
diff --git a/Engine/Plugins/Runtime/AudioModulation/Source/AudioModulation/Private/AudioModulationImpl.cpp b/Engine/Plugins/Runtime/AudioModulation/Source/AudioModulation/Private/AudioModulationImpl.cpp
--- a/Engine/Plugins/Runtime/AudioModulation/Source/AudioModulation/Private/AudioModulationImpl.cpp
+++ b/Engine/Plugins/Runtime/AudioModulation/Source/AudioModulation/Private/AudioModulationImpl.cpp
@@ -107,55 +107,79 @@ namespace AudioModulation
 		const uint32 SettingsId = InSettings.GetUniqueID();
 		RunCommandOnAudioThread([this, SettingsId]()
 		{
-			if (PreviewSound)
+			if (!PreviewSound)
 			{
-				USoundModulationPluginSourceSettingsBase* SettingsBase = PreviewSound->FindModulationSettings();
-				if (SettingsBase && SettingsId == SettingsBase->GetUniqueID())
-				{
-					const USoundModulationSettings* Settings = CastChecked<USoundModulationSettings>(SettingsBase);
+				return;
+			}
 
-					bool bShouldStop = false;
-					if (InputUpdateRequiresStop<FSoundVolumeModulationInput>(Settings->Volume.Inputs, PreviewSettings.Volume.InputProxies))
-					{
-						bShouldStop = true;
-					}
-					else if (InputUpdateRequiresStop<FSoundHPFModulationInput>(Settings->Highpass.Inputs, PreviewSettings.Highpass.InputProxies))
-					{
-						bShouldStop = true;
-					}
-					else if (InputUpdateRequiresStop<FSoundLPFModulationInput>(Settings->Lowpass.Inputs, PreviewSettings.Lowpass.InputProxies))
+			USoundModulationPluginSourceSettingsBase* SettingsBase = PreviewSound->FindModulationSettings();
+			if (!SettingsBase || SettingsId != SettingsBase->GetUniqueID())
+			{
+				return;
+			}
+
+			const USoundModulationSettings* Settings = CastChecked<USoundModulationSettings>(SettingsBase);
+
+			bool bShouldStop = false;
+			if (InputUpdateRequiresStop<FSoundVolumeModulationInput>(Settings->Volume.Inputs, PreviewSettings.Volume.InputProxies))
+			{
+				bShouldStop = true;
+			}
+			else if (InputUpdateRequiresStop<FSoundHPFModulationInput>(Settings->Highpass.Inputs, PreviewSettings.Highpass.InputProxies))
+			{
+				bShouldStop = true;
+			}
+			else if (InputUpdateRequiresStop<FSoundLPFModulationInput>(Settings->Lowpass.Inputs, PreviewSettings.Lowpass.InputProxies))
+			{
+				bShouldStop = true;
+			}
+			else if (InputUpdateRequiresStop<FSoundPitchModulationInput>(Settings->Pitch.Inputs, PreviewSettings.Pitch.InputProxies))
+			{
+				bShouldStop = true;
+			}
+			else
+			{
+				for (int32 i = 0; i < Settings->Controls.Num(); ++i)
+				{
+					const FSoundControlModulationPatch& Control = Settings->Controls[i];
+					if (const FModulationPatchProxy* PatchProxy = PreviewSettings.Controls.Find(Control.Control))
 					{
-						bShouldStop = true;
+						if (InputUpdateRequiresStop<FSoundControlModulationInput>(Control.Inputs, PatchProxy->InputProxies))
+						{
+							bShouldStop = true;
+							break;
+						}
 					}
-					else if (InputUpdateRequiresStop<FSoundPitchModulationInput>(Settings->Pitch.Inputs, PreviewSettings.Pitch.InputProxies))
+					else
 					{
 						bShouldStop = true;
+						break;
 					}
+				}
+			}
 
-					if (bShouldStop)
+			if (bShouldStop)
+			{
+				PreviewSound->Stop();
+				PreviewSound = nullptr;
+				PreviewSettings = FModulationSettingsProxy();
+			}
+			else
+			{
+				PreviewSettings = FModulationSettingsProxy(*Settings);
+				for (FModulationSettingsProxy& SourceSetting : SourceSettings)
+				{
+					if (SourceSetting.GetId() == PreviewSettings.GetId())
 					{
-						PreviewSound->Stop();
-						PreviewSound = nullptr;
-						PreviewSettings = FModulationSettingsProxy();
+						SourceSetting = PreviewSettings;
 					}
-					else
-					{
-						PreviewSettings = FModulationSettingsProxy(*Settings);
-						for (FModulationSettingsProxy& SourceSetting : SourceSettings)
-						{
-							if (SourceSetting.GetId() == PreviewSettings.GetId())
-							{
-								SourceSetting = PreviewSettings;
-							}
-						}
+				}
 
-						for (TPair<uint32, FModulationSettingsProxy>& Pair : SoundSettings)
-						{
-							if (Pair.Value.GetId() == PreviewSettings.GetId())
-							{
-								Pair.Value = PreviewSettings;
-							}
-						}
+				for (TPair<uint32, FModulationSettingsProxy>& Pair : SoundSettings)
+				{
+					if (Pair.Value.GetId() == PreviewSettings.GetId())
+					{
+						Pair.Value = PreviewSettings;
 					}
 				}
 			}
@@ -231,9 +255,12 @@ namespace AudioModulation
 			CheckRefActivate(Cast<USoundControlBusBase>(Input.Bus));
 		}
 
-		for (const FSoundHPFModulationInput& Input : Settings->Highpass.Inputs)
+		for (const FSoundControlModulationPatch& Patch : Settings->Controls)
 		{
-			CheckRefActivate(Cast<USoundControlBusBase>(Input.Bus));
+			for (const FSoundControlModulationInput& Input : Patch.Inputs)
+			{
+				CheckRefActivate(Cast<USoundControlBusBase>(Input.Bus));
+			}
 		}
 	}
 
@@ -286,6 +313,14 @@ namespace AudioModulation
 			{
 				DeactivateBus(Input.BusId, &InSound);
 			}
+
+			for (const TPair<FName, FModulationPatchProxy>& PatchProxy : Settings->Controls)
+			{
+				for (const FModulationInputProxy& Input : PatchProxy.Value.InputProxies)
+				{
+					DeactivateBus(Input.BusId, &InSound);
+				}
+			}
 		}
 
 #if WITH_EDITOR
@@ -319,7 +354,7 @@ namespace AudioModulation
 		const FControlBusProxy NewBusProxy(InBus);
 		const bool bCanCreateNew = InBus.CanAutoActivate(InSound);
 
-		auto ActivateBusInternal = [this, NewBusProxy, InSound, bCanCreateNew]()
+		RunCommandOnAudioThread([this, NewBusProxy, InSound, bCanCreateNew]()
 		{
 			const FBusId NewBusId = NewBusProxy.GetId();
 			FControlBusProxy* BusProxy = ActiveBuses.Find(NewBusId);
@@ -332,9 +367,7 @@ namespace AudioModulation
 			{
 				BusProxy->OnInitSound(*InSound);
 			}
-		};
-
-		IsInAudioThread() ? ActivateBusInternal() : FAudioThread::RunCommandOnAudioThread(ActivateBusInternal);
+		});
 	}
 
 	void FAudioModulationImpl::ActivateBusMix(const USoundControlBusMix& InBusMix, const ISoundModulatable* InSound)
@@ -342,7 +375,7 @@ namespace AudioModulation
 		const FModulatorBusMixProxy NewMixProxy(InBusMix);
 		const bool bCanCreateNew = InBusMix.CanAutoActivate(InSound);
 
-		auto ActivateMixInternal = [this, NewMixProxy, InSound, bCanCreateNew]()
+		RunCommandOnAudioThread([this, NewMixProxy, InSound, bCanCreateNew]()
 		{
 			const FBusMixId MixId = static_cast<FBusMixId>(NewMixProxy.GetId());
 			FModulatorBusMixProxy* BusMixProxy = ActiveBusMixes.Find(MixId);
@@ -360,9 +393,7 @@ namespace AudioModulation
 			{
 				BusMixProxy->OnInitSound(*InSound);
 			}
-		};
-
-		IsInAudioThread() ? ActivateMixInternal() : FAudioThread::RunCommandOnAudioThread(ActivateMixInternal);
+		});
 	}
 
 	void FAudioModulationImpl::ActivateLFO(const USoundBusModulatorLFO& InLFO, const ISoundModulatable* InSound)
@@ -370,7 +401,7 @@ namespace AudioModulation
 		const FModulatorLFOProxy NewLFOProxy(InLFO);
 		const bool bCanCreateNew = InLFO.CanAutoActivate(InSound);
 
-		auto ActivateLFOInternal = [this, NewLFOProxy, InSound, bCanCreateNew]()
+		RunCommandOnAudioThread([this, NewLFOProxy, InSound, bCanCreateNew]()
 		{
 			const AudioModulation::FLFOId LFOId = NewLFOProxy.GetId();
 			FModulatorLFOProxy* LFOProxy = ActiveLFOs.Find(LFOId);
@@ -383,9 +414,7 @@ namespace AudioModulation
 			{
 				LFOProxy->OnInitSound(*InSound);
 			}
-		};
-
-		IsInAudioThread() ? ActivateLFOInternal() : FAudioThread::RunCommandOnAudioThread(ActivateLFOInternal);
+		});
 	}
 
 	float FAudioModulationImpl::CalculateModulationValue(FModulationPatchProxy& OutProxy) const
@@ -461,6 +490,14 @@ namespace AudioModulation
 		return OutValue;
 	}
 
+	bool FAudioModulationImpl::CalculateModulationValue(FModulationPatchProxy& OutProxy, float& OutValue) const
+	{
+		const float InitValue = OutValue;
+		OutValue = CalculateModulationValue(OutProxy);
+
+		return !FMath::IsNearlyEqual(InitValue, OutValue);
+	}
+
 	float FAudioModulationImpl::CalculateInitialVolume(const USoundModulationPluginSourceSettingsBase& SettingsBase) const
 	{
 		check(IsInAudioThread());
@@ -573,16 +610,26 @@ namespace AudioModulation
 		return ActiveLFOs.Contains(InLFOId);
 	}
 
-	void FAudioModulationImpl::ProcessControls(const uint32 SourceId, FSoundModulationControls& Controls)
+	bool FAudioModulationImpl::ProcessControls(const uint32 InSourceId, FSoundModulationControls& OutControls)
 	{
 		check(IsInAudioThread());
 
-		FModulationSettingsProxy& Settings = SourceSettings[SourceId];
+		bool bControlsUpdated = false;
+
+		FModulationSettingsProxy& Settings = SourceSettings[InSourceId];
+
+		bControlsUpdated |= CalculateModulationValue(Settings.Volume, OutControls.Volume);
+		bControlsUpdated |= CalculateModulationValue(Settings.Pitch, OutControls.Pitch);
+		bControlsUpdated |= CalculateModulationValue(Settings.Lowpass, OutControls.Lowpass);
+		bControlsUpdated |= CalculateModulationValue(Settings.Highpass, OutControls.Highpass);
+
+		for (TPair<FName, FModulationPatchProxy>& Pair : Settings.Controls)
+		{
+			float& OutputValue = OutControls.Controls.FindOrAdd(Pair.Key);
+			bControlsUpdated |= CalculateModulationValue(Pair.Value, OutputValue);
+		}
 
-		Controls.Volume   = CalculateModulationValue(Settings.Volume);
-		Controls.Pitch    = CalculateModulationValue(Settings.Pitch);
-		Controls.Lowpass  = CalculateModulationValue(Settings.Lowpass);
-		Controls.Highpass = CalculateModulationValue(Settings.Highpass);
+		return bControlsUpdated;
 	}
 
 	void FAudioModulationImpl::ProcessModulators(float Elapsed)
diff --git a/Engine/Plugins/Runtime/AudioModulation/Source/AudioModulation/Private/AudioModulationInternal.h b/Engine/Plugins/Runtime/AudioModulation/Source/AudioModulation/Private/AudioModulationInternal.h
--- a/Engine/Plugins/Runtime/AudioModulation/Source/AudioModulation/Private/AudioModulationInternal.h
+++ b/Engine/Plugins/Runtime/AudioModulation/Source/AudioModulation/Private/AudioModulationInternal.h
@@ -6,11 +6,12 @@
 #include "Containers/Queue.h"
 
 #include "AudioModulation.h"
+#include "IAudioExtensionPlugin.h"
+#include "SoundModulationProxy.h"
 #include "SoundModulatorBase.h"
 #include "SoundControlBus.h"
 #include "SoundControlBusMix.h"
 #include "SoundModulatorLFO.h"
-#include "IAudioExtensionPlugin.h"
 
 
 #if WITH_AUDIOMODULATION
@@ -54,7 +55,7 @@ namespace AudioModulation
 		bool IsBusActive(const FBusId InBusId) const;
 		bool IsLFOActive(const FLFOId InLFOId) const;
 
-		void ProcessControls(const uint32 InSourceId, FSoundModulationControls& OutControls);
+		bool ProcessControls(const uint32 InSourceId, FSoundModulationControls& OutControls);
 		void ProcessModulators(const float Elapsed);
 
 		void UpdateMix(const USoundControlBusMix& InMix, const TArray<FSoundControlBusMixChannel>& InChannels);
@@ -62,7 +63,12 @@ namespace AudioModulation
 		void UpdateModulator(const USoundModulatorBase& InModulator);
 
 	private:
+		/** Calculates modulation value and returns updated value */
 		float CalculateModulationValue(FModulationPatchProxy& Proxy) const;
+
+		/* Calculates modulation value, storing it in the provided float reference and returns if value changed */
+		bool CalculateModulationValue(FModulationPatchProxy& Proxy, float& OutValue) const;
+
 		void RunCommandOnAudioThread(TFunction<void()> Cmd);
 
 		BusMixProxyMap ActiveBusMixes;
@@ -128,7 +134,7 @@ namespace AudioModulation
 		void DeactivateBus(const FBusId BusId, const ISoundModulatable* Sound = nullptr) { }
 		void DeactivateLFO(const FLFOId LFOId, const ISoundModulatable* Sound = nullptr) { }
 
-		void ProcessControls(const uint32 InSourceId, FSoundModulationControls& OutControls) { }
+		bool ProcessControls(const uint32 InSourceId, FSoundModulationControls& OutControls) { return false; }
 		void ProcessModulators(const float Elapsed) { }
 
 		void UpdateMix(const USoundControlBusMix& InMix, const TArray<FSoundControlBusMixChannel>& InChannels) { }
diff --git a/Engine/Plugins/Runtime/AudioModulation/Source/AudioModulation/Private/AudioModulationSettings.cpp b/Engine/Plugins/Runtime/AudioModulation/Source/AudioModulation/Private/AudioModulationSettings.cpp
new file mode 100644
--- /dev/null
+++ b/Engine/Plugins/Runtime/AudioModulation/Source/AudioModulation/Private/AudioModulationSettings.cpp
@@ -0,0 +1,7 @@
+// Copyright 1998-2019 Epic Games, Inc. All Rights Reserved.
+#include "AudioModulationSettings.h"
+
+UAudioModulationSettings::UAudioModulationSettings(const FObjectInitializer& ObjectInitializer)
+	: Super(ObjectInitializer)
+{
+}
diff --git a/Engine/Plugins/Runtime/AudioModulation/Source/AudioModulation/Private/SoundControlBus.cpp b/Engine/Plugins/Runtime/AudioModulation/Source/AudioModulation/Private/SoundControlBus.cpp
--- a/Engine/Plugins/Runtime/AudioModulation/Source/AudioModulation/Private/SoundControlBus.cpp
+++ b/Engine/Plugins/Runtime/AudioModulation/Source/AudioModulation/Private/SoundControlBus.cpp
@@ -8,7 +8,7 @@
 #include "AudioModulationLogging.h"
 #include "Engine/World.h"
 #include "SoundModulatorLFO.h"
-
+#include "SoundModulationProxy.h"
 
 USoundControlBusBase::USoundControlBusBase(const FObjectInitializer& ObjectInitializer)
 	: Super(ObjectInitializer)
@@ -22,6 +22,16 @@ USoundControlBusBase::USoundControlBusBase(const FObjectInitializer& ObjectIniti
 }
 
 #if WITH_EDITOR
+void USoundControlBusBase::PostDuplicate(EDuplicateMode::Type DuplicateMode)
+{
+	if (!bOverrideAddress)
+	{
+		Address = GetName();
+	}
+
+	Super::PostDuplicate(DuplicateMode);
+}
+
 void USoundControlBusBase::PostEditChangeProperty(FPropertyChangedEvent& PropertyChangedEvent)
 {
 	if (UProperty* Property = PropertyChangedEvent.Property)
@@ -75,6 +85,11 @@ USoundHPFControlBus::USoundHPFControlBus(const FObjectInitializer& ObjectInitial
 	DefaultValue = 0.0f;
 }
 
+USoundControlBus::USoundControlBus(const FObjectInitializer& ObjectInitializer)
+	: Super(ObjectInitializer)
+{
+}
+
 void USoundControlBusBase::BeginDestroy()
 {
 	Super::BeginDestroy();
diff --git a/Engine/Plugins/Runtime/AudioModulation/Source/AudioModulation/Private/SoundControlBus.h b/Engine/Plugins/Runtime/AudioModulation/Source/AudioModulation/Private/SoundControlBus.h
deleted file mode 100644
diff --git a/Engine/Plugins/Runtime/AudioModulation/Source/AudioModulation/Private/SoundControlBusMix.h b/Engine/Plugins/Runtime/AudioModulation/Source/AudioModulation/Private/SoundControlBusMix.h
deleted file mode 100644
diff --git a/Engine/Plugins/Runtime/AudioModulation/Source/AudioModulation/Private/SoundModulationPatch.cpp b/Engine/Plugins/Runtime/AudioModulation/Source/AudioModulation/Private/SoundModulationPatch.cpp
--- a/Engine/Plugins/Runtime/AudioModulation/Source/AudioModulation/Private/SoundModulationPatch.cpp
+++ b/Engine/Plugins/Runtime/AudioModulation/Source/AudioModulation/Private/SoundModulationPatch.cpp
@@ -1,10 +1,12 @@
 // Copyright 1998-2019 Epic Games, Inc. All Rights Reserved.
 #include "SoundModulationPatch.h"
 
+#include "AudioDefines.h"
 #include "AudioDevice.h"
 #include "AudioModulation.h"
 #include "AudioModulationInternal.h"
 #include "AudioModulationStatics.h"
+#include "SoundModulationProxy.h"
 #include "SoundModulationTransform.h"
 
 
@@ -17,24 +19,42 @@ USoundModulationSettings::USoundModulationSettings(const FObjectInitializer& Obj
 	Lowpass.Output.Transform.OutputMin = 20.0f;
 	Lowpass.Output.Transform.OutputMax = 20000.0f;
 	Lowpass.Output.Transform.Curve = ESoundModulatorOutputCurve::Exp;
-	Lowpass.Output.Operator = ESoundModulatorOperator::Min;
+	Lowpass.Output.SetOperator(ESoundModulatorOperator::Min);
 
 	Highpass.DefaultInputValue = 0.0f;
 	Highpass.Output.Transform.OutputMin = 20.0f;
 	Highpass.Output.Transform.OutputMax = 20000.0f;
 	Highpass.Output.Transform.Curve = ESoundModulatorOutputCurve::Exp;
-	Highpass.Output.Operator = ESoundModulatorOperator::Max;
+	Highpass.Output.SetOperator(ESoundModulatorOperator::Max);
 }
 
 #if WITH_EDITOR
 void USoundModulationSettings::OnPostEditChange(UWorld* World)
 {
+	Volume.Clamp();
+	Pitch.Clamp();
+	Lowpass.Clamp();
+	Highpass.Clamp();
+
+	for (FSoundControlModulationPatch& Control : Controls)
+	{
+		Control.Clamp();
+	}
+
 	if (AudioModulation::FAudioModulationImpl* Impl = UAudioModulationStatics::GetModulationImpl(World))
 	{
 		Impl->OnEditPluginSettings(*this);
 	}
 }
 
+void USoundModulationSettings::PostDuplicate(EDuplicateMode::Type DuplicateMode)
+{
+	if (DuplicateMode == EDuplicateMode::Normal)
+	{
+		OnPostEditChange(GetWorld());
+	}
+}
+
 void USoundModulationSettings::PostEditChangeProperty(FPropertyChangedEvent& PropertyChangedEvent)
 {
 	OnPostEditChange(GetWorld());
@@ -44,6 +64,49 @@ void USoundModulationSettings::PostEditChangeChainProperty(FPropertyChangedChain
 {
 	OnPostEditChange(GetWorld());
 }
+
+void FSoundModulationPatchBase::Clamp()
+{
+	if (FSoundModulationOutputBase* Output = GetOutput())
+	{
+		if (Output->Transform.OutputMin > Output->Transform.OutputMax)
+		{
+			Output->Transform.OutputMin = Output->Transform.OutputMax;
+		}
+	}
+}
+
+void FSoundVolumeModulationPatch::Clamp()
+{
+	FSoundModulationPatchBase::Clamp();
+
+	Output.Transform.OutputMin = FMath::Clamp(Output.Transform.OutputMin, 0.0f, MAX_VOLUME);
+	Output.Transform.OutputMax = FMath::Clamp(Output.Transform.OutputMax, 0.0f, MAX_VOLUME);
+}
+
+void FSoundPitchModulationPatch::Clamp()
+{
+	FSoundModulationPatchBase::Clamp();
+
+	Output.Transform.OutputMin = FMath::Clamp(Output.Transform.OutputMin, MIN_PITCH, MAX_PITCH);
+	Output.Transform.OutputMax = FMath::Clamp(Output.Transform.OutputMax, MIN_PITCH, MAX_PITCH);
+}
+
+void FSoundLPFModulationPatch::Clamp()
+{
+	FSoundModulationPatchBase::Clamp();
+
+	Output.Transform.OutputMin = FMath::Clamp(Output.Transform.OutputMin, MIN_FILTER_FREQUENCY, MAX_FILTER_FREQUENCY);
+	Output.Transform.OutputMax = FMath::Clamp(Output.Transform.OutputMax, MIN_FILTER_FREQUENCY, MAX_FILTER_FREQUENCY);
+}
+
+void FSoundHPFModulationPatch::Clamp()
+{
+	FSoundModulationPatchBase::Clamp();
+
+	Output.Transform.OutputMin = FMath::Clamp(Output.Transform.OutputMin, MIN_FILTER_FREQUENCY, MAX_FILTER_FREQUENCY);
+	Output.Transform.OutputMax = FMath::Clamp(Output.Transform.OutputMax, MIN_FILTER_FREQUENCY, MAX_FILTER_FREQUENCY);
+}
 #endif // WITH_EDITOR
 
 void FSoundVolumeModulationPatch::GenerateProxies(TArray<AudioModulation::FModulationInputProxy>& InputProxies) const
@@ -78,11 +141,24 @@ void FSoundLPFModulationPatch::GenerateProxies(TArray<AudioModulation::FModulati
 	}
 }
 
+void FSoundControlModulationPatch::GenerateProxies(TArray<AudioModulation::FModulationInputProxy>& InputProxies) const
+{
+	for (const FSoundControlModulationInput& Input : Inputs)
+	{
+		InputProxies.Emplace_GetRef(Input);
+	}
+}
+
 FSoundModulationOutput::FSoundModulationOutput()
 	: Operator(ESoundModulatorOperator::Multiply)
 {
 }
 
+FSoundModulationOutputFixedOperator::FSoundModulationOutputFixedOperator()
+	: Operator(ESoundModulatorOperator::Multiply)
+{
+}
+
 FSoundModulationInputBase::FSoundModulationInputBase()
 	: bSampleAndHold(0)
 {
@@ -103,6 +179,11 @@ FSoundHPFModulationInput::FSoundHPFModulationInput()
 {
 }
 
+FSoundControlModulationInput::FSoundControlModulationInput()
+	: Bus(nullptr)
+{
+}
+
 FSoundLPFModulationInput::FSoundLPFModulationInput()
 	: Bus(nullptr)
 {
@@ -139,9 +220,9 @@ namespace AudioModulation
 	{
 	}
 
-	FModulationOutputProxy::FModulationOutputProxy(const FSoundModulationOutput& Output)
+	FModulationOutputProxy::FModulationOutputProxy(const FSoundModulationOutputBase& Output)
 		: bInitialized(0)
-		, Operator(Output.Operator)
+		, Operator(Output.GetOperator())
 		, SampleAndHoldValue(1.0f)
 		, Transform(Output.Transform)
 	{
@@ -154,7 +235,7 @@ namespace AudioModulation
 
 	FModulationPatchProxy::FModulationPatchProxy(const FSoundModulationPatchBase& Patch)
 		: DefaultInputValue(Patch.DefaultInputValue)
-		, OutputProxy(Patch.Output)
+		, OutputProxy(*Patch.GetOutput())
 	{
 		Patch.GenerateProxies(InputProxies);
 	}
@@ -183,6 +264,11 @@ namespace AudioModulation
 		, Lowpass(Settings.Lowpass)
 		, Highpass(Settings.Highpass)
 	{
+		for (const FSoundControlModulationPatch& Patch : Settings.Controls)
+		{
+			Controls.Add(Patch.Control, Patch);
+		}
+
 		for (const USoundControlBusMix* Mix : Settings.Mixes)
 		{
 			if (Mix)
diff --git a/Engine/Plugins/Runtime/AudioModulation/Source/AudioModulation/Private/SoundModulationProxy.h b/Engine/Plugins/Runtime/AudioModulation/Source/AudioModulation/Private/SoundModulationProxy.h
new file mode 100644
--- /dev/null
+++ b/Engine/Plugins/Runtime/AudioModulation/Source/AudioModulation/Private/SoundModulationProxy.h
@@ -0,0 +1,288 @@
+// Copyright 1998-2019 Epic Games, Inc. All Rights Reserved.
+#pragma once
+
+#include "CoreMinimal.h"
+#include "IAudioExtensionPlugin.h"
+#include "UObject/ObjectMacros.h"
+#include "UObject/Object.h"
+
+#include "SoundControlBus.h"
+#include "SoundControlBusMix.h"
+#include "SoundModulationPatch.h"
+#include "SoundModulationTransform.h"
+
+
+namespace AudioModulation
+{
+	// Modulator Ids
+	using FBusMixId = uint32;
+	extern const FBusMixId InvalidBusMixId;
+
+	using FBusId = uint32;
+	extern const FBusId InvalidBusId;
+
+	using FLFOId = uint32;
+	extern const FLFOId InvalidLFOId;
+
+
+	// Modulator Proxy Templates
+	template <typename IdType>
+	class TModulatorProxyBase
+	{
+	private:
+		IdType Id;
+
+#if !UE_BUILD_SHIPPING
+		FString Name;
+#endif // !UE_BUILD_SHIPPING
+
+	public:
+		TModulatorProxyBase()
+			: Id(static_cast<IdType>(0))
+		{
+		}
+
+		TModulatorProxyBase(const FString& InName, const uint32 InId)
+			: Id(static_cast<IdType>(InId))
+#if !UE_BUILD_SHIPPING
+			, Name(InName)
+#endif // !UE_BUILD_SHIPPING
+		{
+		}
+
+		IdType GetId() const
+		{
+			return Id;
+		}
+
+		// FOR DEBUG USE ONLY (Not available in shipped builds):
+		// Provides name of object that generated proxy.
+		const FString& GetName() const
+		{
+#if UE_BUILD_SHIPPING
+			static FString Name;
+#endif // UE_BUILD_SHIPPING
+
+			return Name;
+		}
+	};
+
+	template <typename IdType>
+	class TModulatorProxyRefBase : public TModulatorProxyBase<IdType>
+	{
+	public:
+		TModulatorProxyRefBase()
+			: bAutoActivate(false)
+		{
+		}
+
+		TModulatorProxyRefBase(const FString& Name, const IdType InId, const bool bInAutoActivate)
+			: TModulatorProxyBase<IdType>(Name, InId)
+			, bAutoActivate(bInAutoActivate)
+		{
+		}
+
+		virtual ~TModulatorProxyRefBase() = default;
+
+		virtual bool CanDestroy() const
+		{
+			return !bAutoActivate || (bAutoActivate && RefSounds.Num() == 0);
+		}
+
+		int32 OnReleaseSound(const ISoundModulatable& Sound)
+		{
+			if (bAutoActivate)
+			{
+				const int32 NumRemoved = RefSounds.Remove(&Sound);
+				check(NumRemoved == 1);
+			}
+
+			return RefSounds.Num();
+		}
+
+		bool GetAutoActivate() const
+		{
+			return static_cast<bool>(bAutoActivate);
+		}
+
+		const TArray<const ISoundModulatable*>& GetRefSounds() const
+		{
+			return RefSounds;
+		}
+
+		int32 OnInitSound(const ISoundModulatable& Sound)
+		{
+			// Preview sounds force proxies into being auto-activated
+			// to allow for auditioning with the provided modulation settings.
+			if (Sound.IsPreviewSound())
+			{
+				bAutoActivate = 1;
+			}
+
+			if (bAutoActivate)
+			{
+				RefSounds.AddUnique(&Sound);
+			}
+
+			return RefSounds.Num();
+		}
+
+	private:
+		uint8 bAutoActivate : 1;
+		TArray<const ISoundModulatable*> RefSounds;
+	};
+
+	struct FModulatorBusMixChannelProxy : public TModulatorProxyBase<FBusId>
+	{
+		FModulatorBusMixChannelProxy(const FSoundControlBusMixChannel& Channel);
+		FString Address;
+		uint32 ClassId;
+		FSoundModulationValue Value;
+	};
+
+	class FModulatorLFOProxy : public TModulatorProxyRefBase<FLFOId>
+	{
+	public:
+		FModulatorLFOProxy();
+		FModulatorLFOProxy(const USoundBusModulatorLFO& InLFO);
+
+		void OnUpdateProxy(const FModulatorLFOProxy& InLFOProxy);
+
+		float GetValue() const;
+		void Update(float InElapsed);
+
+	private:
+		Audio::FLFO LFO;
+		float Offset;
+		float Value;
+	};
+	using LFOProxyMap = TMap<FBusId, FModulatorLFOProxy>;
+
+	class FControlBusProxy : public TModulatorProxyRefBase<FBusId>
+	{
+	public:
+		FControlBusProxy();
+		FControlBusProxy(const USoundControlBusBase& Bus);
+
+		void OnUpdateProxy(const FControlBusProxy& InBusProxy);
+
+		float GetDefaultValue() const;
+		const TArray<FLFOId>& GetLFOIds() const;
+		float GetLFOValue() const;
+		float GetMixValue() const;
+		FVector2D GetRange() const;
+		float GetValue() const;
+		void MixIn(const float InValue);
+		void MixLFO(LFOProxyMap& LFOMap);
+		void Reset();
+
+	private:
+		float Mix(float ValueA) const;
+		float Mix(float ValueA, float ValueB) const;
+
+		float DefaultValue;
+
+		// Cached values
+		float LFOValue;
+		float MixValue;
+
+		TArray<FLFOId> LFOIds;
+		ESoundModulatorOperator Operator;
+		FVector2D Range;
+	};
+
+	using BusProxyMap = TMap<FBusId, FControlBusProxy>;
+
+
+	class FModulatorBusMixProxy : public TModulatorProxyRefBase<FBusMixId>
+	{
+	public:
+		enum class BusMixStatus : uint8
+		{
+			Enabled,
+			Stopping,
+			Stopped
+		};
+
+		FModulatorBusMixProxy(const USoundControlBusMix& Mix);
+		virtual ~FModulatorBusMixProxy() = default;
+
+		virtual bool CanDestroy() const override;
+
+		virtual void OnUpdateProxy(const FModulatorBusMixProxy& InBusMixProxy);
+
+		void SetEnabled();
+		void SetMix(const TArray<FSoundControlBusMixChannel>& InChannels);
+		void SetMixByFilter(const FString& InAddressFilter, uint32 InFilterClassId, const FSoundModulationValue& InValue);
+		void SetStopping();
+
+		void Update(const float Elapsed, BusProxyMap& ProxyMap);
+
+		TMap<FBusId, FModulatorBusMixChannelProxy> Channels;
+
+	private:
+		BusMixStatus Status;
+	};
+
+	using BusMixProxyMap = TMap<FBusMixId, FModulatorBusMixProxy>;
+
+	/** Modulation input instance */
+	struct FModulationInputProxy
+	{
+		FModulationInputProxy();
+		FModulationInputProxy(const FSoundModulationInputBase& Patch);
+
+		FBusId BusId;
+		FSoundModulationInputTransform Transform;
+		uint8 bSampleAndHold : 1;
+	};
+
+	/** Patch applied as the final stage of a modulation chain prior to output on the sound level (Always active, never removed) */
+	struct FModulationOutputProxy
+	{
+		FModulationOutputProxy();
+		FModulationOutputProxy(const FSoundModulationOutputBase& Patch);
+
+		/** Whether patch has been initialized or not */
+		uint8 bInitialized : 1;
+
+		/** Operator used to calculate the output proxy value */
+		ESoundModulatorOperator Operator;
+
+		/** Cached value of sample-and-hold input values */
+		float SampleAndHoldValue;
+
+		/** Final transform before passing to output */
+		FSoundModulationOutputTransform Transform;
+	};
+
+	struct FModulationPatchProxy
+	{
+		FModulationPatchProxy();
+		FModulationPatchProxy(const FSoundModulationPatchBase& Patch);
+
+		/** Default value of patch (Value mixed when inputs are provided or not, regardless of active state)*/
+		float DefaultInputValue;
+
+		/** Optional modulation inputs */
+		TArray<FModulationInputProxy> InputProxies;
+
+		/** Final output modulation post input combination */
+		FModulationOutputProxy OutputProxy;
+	};
+
+	struct FModulationSettingsProxy : public TModulatorProxyBase<uint32>
+	{
+		FModulationSettingsProxy();
+		FModulationSettingsProxy(const USoundModulationSettings& Settings);
+
+		FModulationPatchProxy Volume;
+		FModulationPatchProxy Pitch;
+		FModulationPatchProxy Lowpass;
+		FModulationPatchProxy Highpass;
+
+		TMap<FName, FModulationPatchProxy> Controls;
+
+		TArray<FBusMixId> Mixes;
+	};
+} // namespace AudioModulation
\ No newline at end of file
diff --git a/Engine/Plugins/Runtime/AudioModulation/Source/AudioModulation/Private/SoundModulationValue.h b/Engine/Plugins/Runtime/AudioModulation/Source/AudioModulation/Private/SoundModulationValue.h
deleted file mode 100644
diff --git a/Engine/Plugins/Runtime/AudioModulation/Source/AudioModulation/Private/SoundModulatorBase.h b/Engine/Plugins/Runtime/AudioModulation/Source/AudioModulation/Private/SoundModulatorBase.h
deleted file mode 100644
diff --git a/Engine/Plugins/Runtime/AudioModulation/Source/AudioModulation/Private/SoundModulatorLFO.h b/Engine/Plugins/Runtime/AudioModulation/Source/AudioModulation/Private/SoundModulatorLFO.h
deleted file mode 100644
diff --git a/Engine/Plugins/Runtime/AudioModulation/Source/AudioModulation/Public/AudioModulation.h b/Engine/Plugins/Runtime/AudioModulation/Source/AudioModulation/Public/AudioModulation.h
--- a/Engine/Plugins/Runtime/AudioModulation/Source/AudioModulation/Public/AudioModulation.h
+++ b/Engine/Plugins/Runtime/AudioModulation/Source/AudioModulation/Public/AudioModulation.h
@@ -41,7 +41,7 @@ namespace AudioModulation
 
 		virtual void OnReleaseSound(ISoundModulatable& Sound) override;
 		virtual void OnReleaseSource(const uint32 SourceId) override;
-		virtual void ProcessControls(const uint32 SourceId, FSoundModulationControls& Controls) override;
+		virtual bool ProcessControls(const uint32 SourceId, FSoundModulationControls& Controls) override;
 		virtual void ProcessModulators(const float Elapsed) override;
 		//~ End IAudioModulation implementation
 
diff --git a/Engine/Plugins/Runtime/AudioModulation/Source/AudioModulation/Public/AudioModulationSettings.h b/Engine/Plugins/Runtime/AudioModulation/Source/AudioModulation/Public/AudioModulationSettings.h
new file mode 100644
--- /dev/null
+++ b/Engine/Plugins/Runtime/AudioModulation/Source/AudioModulation/Public/AudioModulationSettings.h
@@ -0,0 +1,25 @@
+// Copyright 1998-2019 Epic Games, Inc. All Rights Reserved.
+#pragma once
+
+#include "CoreMinimal.h"
+
+#include "Engine/DeveloperSettings.h"
+#include "UObject/ObjectMacros.h"
+#include "UObject/SoftObjectPath.h"
+
+#include "AudioModulationSettings.generated.h"
+
+
+UCLASS(config = Game, defaultconfig, meta = (DisplayName = "Sound Modulation"))
+class AUDIOMODULATION_API UAudioModulationSettings : public UDeveloperSettings
+{
+	GENERATED_UCLASS_BODY()
+
+#if WITH_EDITORONLY_DATA
+	// Names of all controls provided for modulation beyond standard types (Volume, Pitch, HPF, LPF)
+	// Properties hidden as Generic Control Modulation is still in development
+	// UPROPERTY(config, EditAnywhere, Category = "Controls")
+	UPROPERTY()
+	TArray<FName> ControlNames;
+#endif // WITH_EDITORONLY_DATA
+};
\ No newline at end of file
diff --git a/Engine/Plugins/Runtime/AudioModulation/Source/AudioModulation/Public/SoundControlBus.h b/Engine/Plugins/Runtime/AudioModulation/Source/AudioModulation/Public/SoundControlBus.h
new file mode 100644
--- /dev/null
+++ b/Engine/Plugins/Runtime/AudioModulation/Source/AudioModulation/Public/SoundControlBus.h
@@ -0,0 +1,108 @@
+// Copyright 1998-2019 Epic Games, Inc. All Rights Reserved.
+#pragma once
+
+#include "CoreMinimal.h"
+#include "UObject/ObjectMacros.h"
+#include "UObject/Object.h"
+
+#include "SoundModulatorLFO.h"
+#include "SoundModulationValue.h"
+
+#include "SoundControlBus.generated.h"
+
+
+UENUM()
+enum class ESoundModulatorOperator : uint8
+{
+	/** Multiply all mix values together */
+	Multiply,
+
+	/** Take the lowest mix value */
+	Min,
+
+	/** Take the highest mix value */
+	Max,
+
+	Count UMETA(Hidden),
+};
+
+
+UCLASS(BlueprintType, hidecategories = Object, abstract, MinimalAPI)
+class USoundControlBusBase : public USoundModulatorBase
+{
+	GENERATED_UCLASS_BODY()
+
+public:
+#if WITH_EDITORONLY_DATA
+	/** If true, Address field is used in place of object name for address used when applying mix changes using filtering. */
+	UPROPERTY(EditAnywhere, Category = General, BlueprintReadWrite)
+	bool bOverrideAddress;
+#endif // WITH_EDITORONLY_DATA
+
+	/** Address to use when applying mix changes. */
+	UPROPERTY(EditAnywhere, Category = General, BlueprintReadWrite, meta = (EditCondition = "bOverrideAddress"))
+	FString Address;
+
+	/** Default value of modulator (when no mix is applied). */
+	UPROPERTY(EditAnywhere, Category = Modulation, BlueprintReadWrite, meta = (UIMin = "0", UIMax = "1"))
+	float DefaultValue;
+
+	/** Minimum value the bus can achieve (applied post mix phase, pre patch output) */
+	UPROPERTY(EditAnywhere, Category = Modulation, BlueprintReadWrite, meta = (UIMin = "0", UIMax = "1"))
+	float Min;
+
+	/** Maximum value the bus can achieve (applied post mix phase, pre patch output) */
+	UPROPERTY(EditAnywhere, Category = Modulation, BlueprintReadWrite, meta = (UIMin = "0", UIMax = "1"))
+	float Max;
+
+	UPROPERTY(EditAnywhere, Category = Modulation, BlueprintReadWrite)
+	TArray<USoundBusModulatorBase*> Modulators;
+
+#if WITH_EDITOR
+	virtual void PostDuplicate(EDuplicateMode::Type DuplicateMode) override;
+	virtual void PostEditChangeProperty(FPropertyChangedEvent& PropertyChangedEvent) override;
+	virtual void PostInitProperties() override;
+	virtual void PostRename(UObject* OldOuter, const FName OldName) override;
+#endif // WITH_EDITOR
+
+	virtual void BeginDestroy() override;
+	virtual ESoundModulatorOperator GetOperator() const { return ESoundModulatorOperator::Multiply; }
+};
+
+UCLASS(BlueprintType, hidecategories = Object, editinlinenew, MinimalAPI)
+class USoundVolumeControlBus : public USoundControlBusBase
+{
+	GENERATED_UCLASS_BODY()
+};
+
+UCLASS(BlueprintType, hidecategories = Object, editinlinenew, MinimalAPI)
+class USoundPitchControlBus : public USoundControlBusBase
+{
+	GENERATED_UCLASS_BODY()
+
+	virtual ESoundModulatorOperator GetOperator() const { return ESoundModulatorOperator::Multiply; }
+};
+
+UCLASS(BlueprintType, hidecategories = Object, editinlinenew, MinimalAPI)
+class USoundHPFControlBus : public USoundControlBusBase
+{
+	GENERATED_UCLASS_BODY()
+
+	virtual ESoundModulatorOperator GetOperator() const { return ESoundModulatorOperator::Max; }
+};
+
+UCLASS(BlueprintType, hidecategories = Object, editinlinenew, MinimalAPI)
+class USoundLPFControlBus : public USoundControlBusBase
+{
+	GENERATED_UCLASS_BODY()
+
+	virtual ESoundModulatorOperator GetOperator() const { return ESoundModulatorOperator::Min; }
+};
+
+UCLASS(BlueprintType, hidecategories = Object, editinlinenew, MinimalAPI)
+class USoundControlBus : public USoundControlBusBase
+{
+	GENERATED_UCLASS_BODY()
+
+	virtual ESoundModulatorOperator GetOperator() const { return ESoundModulatorOperator::Multiply; }
+};
diff --git a/Engine/Plugins/Runtime/AudioModulation/Source/AudioModulation/Public/SoundControlBusMix.h b/Engine/Plugins/Runtime/AudioModulation/Source/AudioModulation/Public/SoundControlBusMix.h
new file mode 100644
--- /dev/null
+++ b/Engine/Plugins/Runtime/AudioModulation/Source/AudioModulation/Public/SoundControlBusMix.h
@@ -0,0 +1,48 @@
+// Copyright 1998-2019 Epic Games, Inc. All Rights Reserved.
+#pragma once
+
+#include "CoreMinimal.h"
+#include "UObject/ObjectMacros.h"
+#include "UObject/Object.h"
+
+#include "SoundControlBus.h"
+#include "SoundModulatorBase.h"
+#include "SoundModulationValue.h"
+
+#include "SoundControlBusMix.generated.h"
+
+
+// Forward Declarations
+class USoundControlBusBase;
+
+
+USTRUCT(BlueprintType)
+struct FSoundControlBusMixChannel
+{
+	GENERATED_USTRUCT_BODY()
+
+	FSoundControlBusMixChannel();
+	FSoundControlBusMixChannel(USoundControlBusBase* InBus, const float TargetValue);
+
+	/* Bus controlled by channel. */
+	UPROPERTY(EditAnywhere, Category = Channel, BlueprintReadOnly)
+	USoundControlBusBase* Bus;
+
+	/* Value mix is set to. */
+	UPROPERTY(EditAnywhere, Category = Channel, BlueprintReadWrite, meta = (ShowOnlyInnerProperties))
+	FSoundModulationValue Value;
+};
+
+UCLASS(config = Engine, autoexpandcategories = (Channel, Mix), editinlinenew, BlueprintType, MinimalAPI)
+class USoundControlBusMix : public USoundModulatorBase
+{
+	GENERATED_UCLASS_BODY()
+
+public:
+
+	virtual void BeginDestroy() override;
+
+	/* Array of channels controlled by mix. */
+	UPROPERTY(EditAnywhere, Category = Mix, BlueprintReadOnly)
+	TArray<FSoundControlBusMixChannel> Channels;
+};
diff --git a/Engine/Plugins/Runtime/AudioModulation/Source/AudioModulation/Public/SoundModulationPatch.h b/Engine/Plugins/Runtime/AudioModulation/Source/AudioModulation/Public/SoundModulationPatch.h
--- a/Engine/Plugins/Runtime/AudioModulation/Source/AudioModulation/Public/SoundModulationPatch.h
+++ b/Engine/Plugins/Runtime/AudioModulation/Source/AudioModulation/Public/SoundModulationPatch.h
@@ -5,8 +5,6 @@
 #include "UObject/ObjectMacros.h"
 #include "UObject/Object.h"
 #include "SoundModulationTransform.h"
-#include "SoundModulationValue.h"
-#include "SoundModulatorBase.h"
 #include "SoundControlBus.h"
 #include "SoundControlBusMix.h"
 
@@ -21,19 +19,51 @@ namespace AudioModulation
 
 
 USTRUCT(BlueprintType)
-struct FSoundModulationOutput
+struct FSoundModulationOutputBase
 {
 	GENERATED_USTRUCT_BODY()
 
-	FSoundModulationOutput();
+	virtual ~FSoundModulationOutputBase() = default;
+
+	/** Final transform before passing to output */
+	UPROPERTY(EditAnywhere, Category = Modulation, BlueprintReadWrite, meta = (ShowOnlyInnerProperties))
+	FSoundModulationOutputTransform Transform;
+
+	virtual ESoundModulatorOperator GetOperator() const PURE_VIRTUAL(FSoundModulationOutputBase::GetOperator, return ESoundModulatorOperator::Multiply;);
+};
+
+USTRUCT(BlueprintType)
+struct FSoundModulationOutputFixedOperator : public FSoundModulationOutputBase
+{
+	GENERATED_USTRUCT_BODY()
 
+	FSoundModulationOutputFixedOperator();
+
+protected:
 	/** Operator used when mixing input values */
 	UPROPERTY(VisibleAnywhere, BlueprintReadOnly, Category = Modulation)
 	ESoundModulatorOperator Operator;
 
-	/** Final transform before passing to output */
-	UPROPERTY(EditAnywhere, Category = Modulation, BlueprintReadWrite, meta = (ShowOnlyInnerProperties))
-	FSoundModulationOutputTransform Transform;
+public:
+	ESoundModulatorOperator GetOperator() const override { return Operator; }
+	void SetOperator(ESoundModulatorOperator InOperator) { Operator = InOperator; }
+};
+
+USTRUCT(BlueprintType)
+struct FSoundModulationOutput : public FSoundModulationOutputBase
+{
+	GENERATED_USTRUCT_BODY()
+
+	FSoundModulationOutput();
+
+protected:
+	/** Operator used when mixing input values */
+	UPROPERTY(EditAnywhere, BlueprintReadOnly, Category = Modulation)
+	ESoundModulatorOperator Operator;
+
+public:
+	ESoundModulatorOperator GetOperator() const override { return Operator; }
+	void SetOperator(ESoundModulatorOperator InOperator) { Operator = InOperator; }
 };
 
 USTRUCT(BlueprintType)
@@ -111,6 +141,20 @@ struct FSoundHPFModulationInput : public FSoundModulationInputBase
 	virtual const USoundControlBusBase* GetBus() const { return Cast<USoundControlBusBase>(Bus); }
 };
 
+USTRUCT(BlueprintType)
+struct FSoundControlModulationInput : public FSoundModulationInputBase
+{
+	GENERATED_USTRUCT_BODY()
+
+	FSoundControlModulationInput();
+
+	/** The input bus */
+	UPROPERTY(EditAnywhere, BlueprintReadWrite, Category = Input)
+	USoundControlBus* Bus;
+
+	virtual const USoundControlBusBase* GetBus() const { return Cast<USoundControlBusBase>(Bus); }
+};
+
 USTRUCT(BlueprintType)
 struct FSoundModulationPatchBase
 {
@@ -124,11 +168,13 @@ struct FSoundModulationPatchBase
 	UPROPERTY(EditAnywhere, BlueprintReadWrite, Category = Inputs, meta = (UIMin = "0", UIMax = "1"))
 	float DefaultInputValue;
 
-	/** Final modulation parameters to apply */
-	UPROPERTY(EditAnywhere, BlueprintReadWrite, Category = Output, meta = (ShowOnlyInnerProperties))
-	FSoundModulationOutput Output;
-
 	virtual void GenerateProxies(TArray<AudioModulation::FModulationInputProxy>& InputProxies) const PURE_VIRTUAL(FSoundModulationPatchBase::GenerateProxies, );
+	virtual FSoundModulationOutputBase* GetOutput() PURE_VIRTUAL(FSoundModulationPatchBase::GetOutput(), return nullptr;);
+	virtual const FSoundModulationOutputBase* GetOutput() const PURE_VIRTUAL(FSoundModulationPatchBase::GetOutput(), return nullptr;);
+
+#if WITH_EDITOR
+	virtual void Clamp();
+#endif // WITH_EDITOR
 };
 
 USTRUCT(BlueprintType)
@@ -140,7 +186,17 @@ struct FSoundVolumeModulationPatch : public FSoundModulationPatchBase
 	UPROPERTY(EditAnywhere, BlueprintReadWrite, Category = Inputs)
 	TArray<FSoundVolumeModulationInput> Inputs;
 
+	/** Final modulation parameters to apply */
+	UPROPERTY(EditAnywhere, BlueprintReadWrite, Category = Output, meta = (ShowOnlyInnerProperties))
+	FSoundModulationOutputFixedOperator Output;
+
 	virtual void GenerateProxies(TArray<AudioModulation::FModulationInputProxy>& InputProxies) const override;
+	virtual FSoundModulationOutputBase* GetOutput() override { return static_cast<FSoundModulationOutputBase*>(&Output); }
+	virtual const FSoundModulationOutputBase* GetOutput() const override { return static_cast<const FSoundModulationOutputBase*>(&Output); }
+
+#if WITH_EDITOR
+	virtual void Clamp() override;
+#endif // WITH_EDITOR
 };
 
 USTRUCT(BlueprintType)
@@ -152,7 +208,17 @@ struct FSoundPitchModulationPatch : public FSoundModulationPatchBase
 	UPROPERTY(EditAnywhere, BlueprintReadWrite, Category = Inputs)
 	TArray<FSoundPitchModulationInput> Inputs;
 
+	/** Final modulation parameters to apply */
+	UPROPERTY(EditAnywhere, BlueprintReadWrite, Category = Output, meta = (ShowOnlyInnerProperties))
+	FSoundModulationOutputFixedOperator Output;
+
 	virtual void GenerateProxies(TArray<AudioModulation::FModulationInputProxy>& InputProxies) const override;
+	virtual FSoundModulationOutputBase* GetOutput() override { return static_cast<FSoundModulationOutputBase*>(&Output); }
+	virtual const FSoundModulationOutputBase* GetOutput() const override { return static_cast<const FSoundModulationOutputBase*>(&Output); }
+
+#if WITH_EDITOR
+	virtual void Clamp() override;
+#endif // WITH_EDITOR
 };
 
 USTRUCT(BlueprintType)
@@ -164,7 +230,17 @@ struct FSoundLPFModulationPatch : public FSoundModulationPatchBase
 	UPROPERTY(EditAnywhere, BlueprintReadWrite, Category = Inputs)
 	TArray<FSoundLPFModulationInput> Inputs;
 
+	/** Final modulation parameters to apply */
+	UPROPERTY(EditAnywhere, BlueprintReadWrite, Category = Output, meta = (ShowOnlyInnerProperties))
+	FSoundModulationOutputFixedOperator Output;
+
 	virtual void GenerateProxies(TArray<AudioModulation::FModulationInputProxy>& InputProxies) const override;
+	virtual FSoundModulationOutputBase* GetOutput() override { return static_cast<FSoundModulationOutputBase*>(&Output); }
+	virtual const FSoundModulationOutputBase* GetOutput() const override { return static_cast<const FSoundModulationOutputBase*>(&Output); }
+
+#if WITH_EDITOR
+	virtual void Clamp() override;
+#endif // WITH_EDITOR
 };
 
 USTRUCT(BlueprintType)
@@ -176,7 +252,39 @@ struct FSoundHPFModulationPatch : public FSoundModulationPatchBase
 	UPROPERTY(EditAnywhere, BlueprintReadWrite, Category = Inputs)
 	TArray<FSoundHPFModulationInput> Inputs;
 
+	/** Final modulation parameters to apply */
+	UPROPERTY(EditAnywhere, BlueprintReadWrite, Category = Output, meta = (ShowOnlyInnerProperties))
+	FSoundModulationOutputFixedOperator Output;
+
 	virtual void GenerateProxies(TArray<AudioModulation::FModulationInputProxy>& InputProxies) const override;
+	virtual FSoundModulationOutputBase* GetOutput() override { return static_cast<FSoundModulationOutputBase*>(&Output); }
+	virtual const FSoundModulationOutputBase* GetOutput() const override { return static_cast<const FSoundModulationOutputBase*>(&Output); }
+
+#if WITH_EDITOR
+	virtual void Clamp() override;
+#endif // WITH_EDITOR
+};
+
+USTRUCT(BlueprintType)
+struct FSoundControlModulationPatch : public FSoundModulationPatchBase
+{
+	GENERATED_USTRUCT_BODY()
+
+	/** Name of modulation control for sounds referencing this ModulationSettings asset. */
+	UPROPERTY(EditAnywhere, BlueprintReadWrite, Category = Inputs)
+	FName Control;
+
+	/** Modulation inputs */
+	UPROPERTY(EditAnywhere, BlueprintReadWrite, Category = Inputs)
+	TArray<FSoundControlModulationInput> Inputs;
+
+	/** Final modulation parameters to apply */
+	UPROPERTY(EditAnywhere, BlueprintReadWrite, Category = Output, meta = (ShowOnlyInnerProperties))
+	FSoundModulationOutput Output;
+
+	virtual void GenerateProxies(TArray<AudioModulation::FModulationInputProxy>& InputProxies) const override;
+	virtual FSoundModulationOutputBase* GetOutput() override { return static_cast<FSoundModulationOutputBase*>(&Output); }
+	virtual const FSoundModulationOutputBase* GetOutput() const override { return static_cast<const FSoundModulationOutputBase*>(&Output); }
 };
 
 UCLASS(config = Engine, editinlinenew, BlueprintType, MinimalAPI, autoExpandCategories = (Volume, Pitch, Highpass, Lowpass))
@@ -197,6 +305,11 @@ public:
 	UPROPERTY(EditAnywhere, BlueprintReadWrite, Category = Lowpass)
 	FSoundLPFModulationPatch Lowpass;
 
+	// Array of named generic controls for use with modulateable parameters on source effects
+	// Properties hidden as Generic Control Modulation is still in development
+	UPROPERTY()
+	TArray<FSoundControlModulationPatch> Controls;
+
 	// Mixes that will applied and removed when sounds utilizing settings
 	// play and stop respectively. If mix has already been applied manually,
 	// mix will be removed once all sound settings referencing mix stop. Manual
@@ -207,69 +320,8 @@ public:
 #if WITH_EDITOR
 	AUDIOMODULATION_API void OnPostEditChange(UWorld* World);
 
+	virtual void PostDuplicate(EDuplicateMode::Type DuplicateMode) override;
 	virtual void PostEditChangeProperty(FPropertyChangedEvent& PropertyChangedEvent) override;
 	virtual void PostEditChangeChainProperty(FPropertyChangedChainEvent& PropertyChangedEvent) override;
-#endif
+#endif // WITH_EDITOR
 };
-
-
-namespace AudioModulation
-{
-	/** Modulation input instance */
-	struct FModulationInputProxy
-	{
-		FModulationInputProxy();
-		FModulationInputProxy(const FSoundModulationInputBase& Patch);
-
-		FBusId BusId;
-		FSoundModulationInputTransform Transform;
-		uint8 bSampleAndHold : 1;
-	};
-
-	/** Patch applied as the final stage of a modulation chain prior to output on the sound level (Always active, never removed) */
-	struct FModulationOutputProxy
-	{
-		FModulationOutputProxy();
-		FModulationOutputProxy(const FSoundModulationOutput& Patch);
-
-		/** Whether patch has been initialized or not */
-		uint8 bInitialized : 1;
-
-		/** Operator used to calculate the output proxy value */
-		ESoundModulatorOperator Operator;
-
-		/** Cached value of sample-and-hold input values */
-		float SampleAndHoldValue;
-
-		/** Final transform before passing to output */
-		FSoundModulationOutputTransform Transform;
-	};
-
-	struct FModulationPatchProxy
-	{
-		FModulationPatchProxy();
-		FModulationPatchProxy(const FSoundModulationPatchBase& Patch);
-
-		/** Default value of patch (value used when inputs are either not provided or not active)*/
-		float DefaultInputValue;
-
-		/** Optional modulation inputs */
-		TArray<FModulationInputProxy> InputProxies;
-
-		/** Final output modulation post input combination */
-		FModulationOutputProxy OutputProxy;
-	};
-
-	struct FModulationSettingsProxy : TModulatorProxyBase<uint32>
-	{
-		FModulationSettingsProxy();
-		FModulationSettingsProxy(const USoundModulationSettings& Settings);
-
-		FModulationPatchProxy Volume;
-		FModulationPatchProxy Pitch;
-		FModulationPatchProxy Lowpass;
-		FModulationPatchProxy Highpass;
-
-		TArray<FBusMixId> Mixes;
-	};
-} // namespace AudioModulation
diff --git a/Engine/Plugins/Runtime/AudioModulation/Source/AudioModulation/Public/SoundModulationValue.h b/Engine/Plugins/Runtime/AudioModulation/Source/AudioModulation/Public/SoundModulationValue.h
new file mode 100644
--- /dev/null
+++ b/Engine/Plugins/Runtime/AudioModulation/Source/AudioModulation/Public/SoundModulationValue.h
@@ -0,0 +1,37 @@
+// Copyright 1998-2019 Epic Games, Inc. All Rights Reserved.
+#pragma once
+
+#include "UObject/ObjectMacros.h"
+#include "UObject/Object.h"
+
+#include "SoundModulationValue.generated.h"
+
+
+USTRUCT(BlueprintType)
+struct FSoundModulationValue
+{
+	GENERATED_USTRUCT_BODY()
+
+	FSoundModulationValue();
+	FSoundModulationValue(float InValue, float InAttackTime, float InReleaseTime);
+
+	/** Target value of the modulator. */
+	UPROPERTY(EditAnywhere, Category = Modulation, BlueprintReadWrite, meta = (DisplayName = "Value", ClampMin = "0.0", UIMin = "0.0", ClampMax = "1.0", UIMax = "1.0"))
+	float TargetValue;
+
+	/** Time it takes (in sec) to unitarily increase the bus value (from 0 to 1). */
+	UPROPERTY(EditAnywhere, Category = Modulation, BlueprintReadWrite, meta = (DisplayName = "Attack Time", ClampMin = "0.0", UIMin = "0.0"))
+	float AttackTime;
+
+	/** Time it takes (in sec) to unitarily decrease the bus value (from 1 to 0). */
+	UPROPERTY(EditAnywhere, Category = Modulation, BlueprintReadWrite, meta = (DisplayName = "Release Time", ClampMin = "0.0", UIMin = "0.0"))
+	float ReleaseTime;
+
+	/** Current value lerping toward target */
+	float GetCurrentValue() const;
+
+	void Update(float Elapsed);
+
+private:
+	float Value;
+};
diff --git a/Engine/Plugins/Runtime/AudioModulation/Source/AudioModulation/Public/SoundModulatorBase.h b/Engine/Plugins/Runtime/AudioModulation/Source/AudioModulation/Public/SoundModulatorBase.h
new file mode 100644
--- /dev/null
+++ b/Engine/Plugins/Runtime/AudioModulation/Source/AudioModulation/Public/SoundModulatorBase.h
@@ -0,0 +1,43 @@
+// Copyright 1998-2019 Epic Games, Inc. All Rights Reserved.
+#pragma once
+
+#include "CoreMinimal.h"
+#include "IAudioExtensionPlugin.h"
+
+#include "SoundModulatorBase.generated.h"
+
+
+// Forward Declarations
+class ISoundModulatable;
+
+
+/**
+ * Base class for all modulators
+ */
+UCLASS(hideCategories = Object, abstract, MinimalAPI)
+class USoundModulatorBase : public UObject
+{
+	GENERATED_UCLASS_BODY()
+
+public:
+	/**
+	 * When set, automatically activates/deactivates when sounds referencing asset begin/end playing respectively.
+	 */
+	UPROPERTY(EditAnywhere, Category = General, BlueprintReadWrite)
+	uint8 bAutoActivate : 1;
+
+	bool CanAutoActivate(const ISoundModulatable* InSound) const;
+
+#if WITH_EDITOR
+	virtual void PostEditChangeProperty(FPropertyChangedEvent& PropertyChangedEvent) override;
+#endif // WITH_EDITOR
+};
+
+/**
+ * Base class for modulators that manipulate control bus values
+ */
+UCLASS(hideCategories = Object, abstract, MinimalAPI)
+class USoundBusModulatorBase : public USoundModulatorBase
+{
+	GENERATED_UCLASS_BODY()
+};
diff --git a/Engine/Plugins/Runtime/AudioModulation/Source/AudioModulation/Public/SoundModulatorLFO.h b/Engine/Plugins/Runtime/AudioModulation/Source/AudioModulation/Public/SoundModulatorLFO.h
new file mode 100644
--- /dev/null
+++ b/Engine/Plugins/Runtime/AudioModulation/Source/AudioModulation/Public/SoundModulatorLFO.h
@@ -0,0 +1,55 @@
+// Copyright 1998-2019 Epic Games, Inc. All Rights Reserved.
+#pragma once
+
+#include "CoreMinimal.h"
+#include "DSP/LFO.h"
+#include "UObject/ObjectMacros.h"
+#include "UObject/Object.h"
+
+#include "SoundModulatorBase.h"
+#include "SoundModulationValue.h"
+
+#include "SoundModulatorLFO.generated.h"
+
+UENUM(BlueprintType)
+enum class ESoundModulatorLFOShape : uint8
+{
+	Sine			  UMETA(DisplayName = "Sine"),
+	UpSaw			  UMETA(DisplayName = "Saw (Up)"),
+	DownSaw			  UMETA(DisplayName = "Saw (Down)"),
+	Square			  UMETA(DisplayName = "Square"),
+	Triangle		  UMETA(DisplayName = "Triangle"),
+	Exponential		  UMETA(DisplayName = "Exponential"),
+	RandomSampleHold  UMETA(DisplayName = "Random"),
+
+	COUNT UMETA(Hidden)
+};
+
+UCLASS(BlueprintType, hidecategories = Object, editinlinenew, MinimalAPI)
+class USoundBusModulatorLFO : public USoundBusModulatorBase
+{
+	GENERATED_UCLASS_BODY()
+
+public:
+	virtual void BeginDestroy() override;
+
+	/** Shape of oscillating waveform */
+	UPROPERTY(EditAnywhere, Category = Modulation, BlueprintReadWrite)
+	ESoundModulatorLFOShape Shape;
+
+	/** Amplitude of oscillator */
+	UPROPERTY(EditAnywhere, Category = Modulation, BlueprintReadWrite, meta = (UIMin = "0", UIMax = "1", ClampMin = "0", ClampMax = "1"))
+	float Amplitude;
+
+	/** Frequency of oscillator */
+	UPROPERTY(EditAnywhere, Category = Modulation, BlueprintReadWrite, meta = (UIMin = "0", UIMax = "20", ClampMin = "0", ClampMax = "20"))
+	float Frequency;
+
+	/** Offset of oscillator */
+	UPROPERTY(EditAnywhere, Category = Modulation, BlueprintReadWrite, meta = (UIMin = "0", UIMax = "1", ClampMin = "0", ClampMax = "1"))
+	float Offset;
+
+	/** Whether or not to loop the oscillation more than once */
+	UPROPERTY(EditAnywhere, Category = Modulation, BlueprintReadWrite)
+	uint8 bLooping : 1;
+};
diff --git a/Engine/Plugins/Runtime/AudioModulation/Source/AudioModulationEditor/Private/Layouts/SoundModulationTransformLayout.cpp b/Engine/Plugins/Runtime/AudioModulation/Source/AudioModulationEditor/Private/Layouts/SoundModulationTransformLayout.cpp
--- a/Engine/Plugins/Runtime/AudioModulation/Source/AudioModulationEditor/Private/Layouts/SoundModulationTransformLayout.cpp
+++ b/Engine/Plugins/Runtime/AudioModulation/Source/AudioModulationEditor/Private/Layouts/SoundModulationTransformLayout.cpp
@@ -38,8 +38,6 @@ bool FSoundModulationOutputTransformLayoutCustomization::IsCurveSet(TSharedPtr<I
 
 void FSoundModulationOutputTransformLayoutCustomization::CustomizeHeader(TSharedRef<IPropertyHandle> PropertyHandle, FDetailWidgetRow& HeaderRow, IPropertyTypeCustomizationUtils& CustomizationUtils)
 {
-	// Property handle here is the base struct. We are going to hide it since we're showing it's properties directly.
-// 	PropertyHandle->MarkHiddenByCustomization();
 }
 
 void FSoundModulationOutputTransformLayoutCustomization::CustomizeChildren(TSharedRef<IPropertyHandle> StructPropertyHandle, IDetailChildrenBuilder& ChildBuilder, IPropertyTypeCustomizationUtils& StructCustomizationUtils)
@@ -61,21 +59,8 @@ void FSoundModulationOutputTransformLayoutCustomization::CustomizeChildren(TShar
 
 	TSharedRef<IPropertyHandle>InputMaxHandle = PropertyHandles.FindChecked(GET_MEMBER_NAME_CHECKED(FSoundModulationOutputTransform, InputMax)).ToSharedRef();
 	ChildBuilder.AddProperty(InputMaxHandle);
-
-	TSharedRef<IPropertyHandle>CurveHandle = PropertyHandles.FindChecked(GET_MEMBER_NAME_CHECKED(FSoundModulationOutputTransform, Curve)).ToSharedRef();
-	ChildBuilder.AddProperty(CurveHandle);
-
-	TSharedRef<IPropertyHandle> FloatCurveHandle = PropertyHandles.FindChecked(GET_MEMBER_NAME_CHECKED(FSoundModulationOutputTransform, CurveCustom)).ToSharedRef();
-	const TArray<FString> CustomFilters = TArray<FString>({ TEXT("Custom") });
-	ChildBuilder.AddProperty(FloatCurveHandle)
-		.EditCondition(TAttribute<bool>::Create([this, CurveHandle, CustomFilters]() { return IsCurveSet(CurveHandle, CustomFilters); }), nullptr)
-		.Visibility(TAttribute<EVisibility>::Create([this, CurveHandle, CustomFilters]() { return IsCurveSet(CurveHandle, CustomFilters) ? EVisibility::Visible : EVisibility::Hidden; }));
-
 	TSharedRef<IPropertyHandle> ExpScalarHandle = PropertyHandles.FindChecked(GET_MEMBER_NAME_CHECKED(FSoundModulationOutputTransform, Scalar)).ToSharedRef();
-	const TArray<FString> ScalarFilters = TArray<FString>({ TEXT("Exp"), TEXT("Log") });
-	ChildBuilder.AddProperty(ExpScalarHandle)
-		.EditCondition(TAttribute<bool>::Create([this, CurveHandle, ScalarFilters]() { return IsCurveSet(CurveHandle, ScalarFilters); }), nullptr)
-		.Visibility(TAttribute<EVisibility>::Create([this, CurveHandle, ScalarFilters]() { return IsCurveSet(CurveHandle, ScalarFilters) ? EVisibility::Visible : EVisibility::Hidden; }));
+	ChildBuilder.AddProperty(ExpScalarHandle);
 
 	TSharedRef<IPropertyHandle>OutputMinHandle = PropertyHandles.FindChecked(GET_MEMBER_NAME_CHECKED(FSoundModulationOutputTransform, OutputMin)).ToSharedRef();
 	ChildBuilder.AddProperty(OutputMinHandle);
diff --git a/Engine/Source/Runtime/AudioMixer/Private/AudioMixerDevice.cpp b/Engine/Source/Runtime/AudioMixer/Private/AudioMixerDevice.cpp
--- a/Engine/Source/Runtime/AudioMixer/Private/AudioMixerDevice.cpp
+++ b/Engine/Source/Runtime/AudioMixer/Private/AudioMixerDevice.cpp
@@ -898,6 +898,11 @@ namespace Audio
 		});
 	}
 
+	void FMixerDevice::UpdateModulationControls(const uint32 InSourceId, const FSoundModulationControls& InControls)
+	{
+		SourceManager.UpdateModulationControls(InSourceId, InControls);
+	}
+
 	void FMixerDevice::UpdateSourceEffectChain(const uint32 SourceEffectChainId, const TArray<FSourceEffectChainEntry>& SourceEffectChain, const bool bPlayEffectChainTails)
 	{
 		TArray<FSourceEffectChainEntry>* ExistingOverride = SourceEffectChainOverrides.Find(SourceEffectChainId);
diff --git a/Engine/Source/Runtime/AudioMixer/Private/AudioMixerSource.cpp b/Engine/Source/Runtime/AudioMixer/Private/AudioMixerSource.cpp
--- a/Engine/Source/Runtime/AudioMixer/Private/AudioMixerSource.cpp
+++ b/Engine/Source/Runtime/AudioMixer/Private/AudioMixerSource.cpp
@@ -821,7 +821,13 @@ namespace Audio
 
 		if (AudioDevice->IsModulationPluginEnabled())
 		{
-			AudioDevice->ModulationInterface->ProcessControls(MixerSourceVoice->GetSourceId(), WaveInstance->SoundModulationControls);
+			const int32 SourceId = MixerSourceVoice->GetSourceId();
+			const bool bUpdatePending = AudioDevice->ModulationInterface->ProcessControls(SourceId, WaveInstance->SoundModulationControls);
+
+			if (bUpdatePending)
+			{
+				AudioDevice->UpdateModulationControls(SourceId, WaveInstance->SoundModulationControls);
+			}
 		}
 	}
 
diff --git a/Engine/Source/Runtime/AudioMixer/Private/AudioMixerSourceManager.cpp b/Engine/Source/Runtime/AudioMixer/Private/AudioMixerSourceManager.cpp
--- a/Engine/Source/Runtime/AudioMixer/Private/AudioMixerSourceManager.cpp
+++ b/Engine/Source/Runtime/AudioMixer/Private/AudioMixerSourceManager.cpp
@@ -232,6 +232,7 @@ namespace Audio
 			SourceInfo.bIsVorbis = false;
 			SourceInfo.bIsBypassingLPF = false;
 			SourceInfo.bIsBypassingHPF = false;
+			SourceInfo.bIsModulationUpdated = false;
 
 #if AUDIO_MIXER_ENABLE_DEBUG_MODE
 			SourceInfo.bIsDebugMode = false;
@@ -1013,6 +1014,22 @@ namespace Audio
 		});
 	}
 
+	void FMixerSourceManager::UpdateModulationControls(const int32 SourceId, const FSoundModulationControls& InControls)
+	{
+		AUDIO_MIXER_CHECK(SourceId < NumTotalSources);
+		AUDIO_MIXER_CHECK(GameThreadInfo.bIsBusy[SourceId]);
+		AUDIO_MIXER_CHECK_GAME_THREAD(MixerDevice);
+
+		const FSoundModulationControls UpdatedControls = InControls;
+		AudioMixerThreadCommand([this, SourceId, UpdatedControls]()
+		{
+			AUDIO_MIXER_CHECK_AUDIO_PLAT_THREAD(MixerDevice);
+
+			SourceInfos[SourceId].ModulationControls = UpdatedControls;
+			SourceInfos[SourceId].bIsModulationUpdated = true;
+		});
+	}
+
 	void FMixerSourceManager::SetSpatializationParams(const int32 SourceId, const FSpatializationParams& InParams)
 	{
 		AUDIO_MIXER_CHECK(SourceId < NumTotalSources);
@@ -2074,15 +2091,29 @@ namespace Audio
 				// Loop through the effect chain passing in buffers
 				for (FSoundEffectSource* SoundEffectSource : SourceInfo.SourceEffects)
 				{
+					bool bPresetUpdated = false;
+					if (SoundEffectSource->IsActive())
+					{
+						bPresetUpdated = SoundEffectSource->Update();
+					}
+
+					// Modulation must be updated regardless of whether or not the source is
+					// active to allow for initial conditions to be set if source is reactivated.
+					if (bPresetUpdated || SourceInfo.bIsModulationUpdated)
+					{
+						SoundEffectSource->ProcessControls(SourceInfo.ModulationControls);
+					}
+
 					if (SoundEffectSource->IsActive())
 					{
-						SoundEffectSource->Update();
 						SoundEffectSource->ProcessAudio(SourceInfo.SourceEffectInputData, OutputSourceEffectBufferPtr);
 
 						// Copy output to input
 						FMemory::Memcpy(SourceInfo.SourceEffectInputData.InputSourceEffectBufferPtr, OutputSourceEffectBufferPtr, sizeof(float)*NumSamples);
 					}
 				}
+
+				SourceInfo.bIsModulationUpdated = false;
 			}
 
 			const bool bWasEffectTailsDone = SourceInfo.bEffectTailsDone;
@@ -2158,6 +2189,8 @@ namespace Audio
 				SourceInfo.NextFrameValues.Reset();
 				SourceInfo.CurrentPCMBuffer = nullptr;
 			}
+
+			SourceInfo.bIsModulationUpdated = false;
 		}
 	}
 
diff --git a/Engine/Source/Runtime/AudioMixer/Private/AudioMixerSourceManager.h b/Engine/Source/Runtime/AudioMixer/Private/AudioMixerSourceManager.h
--- a/Engine/Source/Runtime/AudioMixer/Private/AudioMixerSourceManager.h
+++ b/Engine/Source/Runtime/AudioMixer/Private/AudioMixerSourceManager.h
@@ -247,6 +247,7 @@ namespace Audio
 		void UpdateDeviceChannelCount(const int32 InNumOutputChannels);
 
 		void UpdateSourceEffectChain(const uint32 SourceEffectChainId, const TArray<FSourceEffectChainEntry>& SourceEffectChain, const bool bPlayEffectChainTails);
+		void UpdateModulationControls(const int32 SourceId, const FSoundModulationControls& InControls);
 
 		const float* GetPreDistanceAttenuationBuffer(const int32 SourceId) const;
 		const float* GetPreEffectBuffer(const int32 SourceId) const;
@@ -480,6 +481,9 @@ namespace Audio
 			Audio::FEnvelopeFollower SourceEnvelopeFollower;
 			float SourceEnvelopeValue;
 
+			// Modulation control state
+			FSoundModulationControls ModulationControls;
+
 			FSpatializationParams SpatParams;
 			Audio::AlignedFloatBuffer ScratchChannelMap;
 
@@ -502,7 +506,7 @@ namespace Audio
 			uint8 bIsVorbis:1;
 			uint8 bIsBypassingLPF:1;
 			uint8 bIsBypassingHPF:1;
-
+			uint8 bIsModulationUpdated:1;
 			// Source format info
 			int32 NumInputChannels;
 			int32 NumPostEffectChannels;
@@ -534,7 +538,7 @@ namespace Audio
 
 		// Array of source infos.
 		TArray<FSourceInfo> SourceInfos;
-		
+
 		// These structs are used for guaranteed vectorization when downmixing
 		// sources.
 		TArray<FSourceDownmixData> DownmixDataArray;
diff --git a/Engine/Source/Runtime/AudioMixer/Public/AudioMixerDevice.h b/Engine/Source/Runtime/AudioMixer/Public/AudioMixerDevice.h
--- a/Engine/Source/Runtime/AudioMixer/Public/AudioMixerDevice.h
+++ b/Engine/Source/Runtime/AudioMixer/Public/AudioMixerDevice.h
@@ -105,6 +105,9 @@ namespace Audio
 		virtual void InitSoundEffectPresets() override;
 		virtual int32 GetNumActiveSources() const override;
 
+		// Updates the provided source's modulation control values
+		virtual void UpdateModulationControls(const uint32 InSourceId, const FSoundModulationControls& InControls) override;
+
 		// Updates the source effect chain (using unique object id). 
 		virtual void UpdateSourceEffectChain(const uint32 SourceEffectChainId, const TArray<FSourceEffectChainEntry>& SourceEffectChain, const bool bPlayEffectChainTails) override;
 		virtual bool GetCurrentSourceEffectChain(const uint32 SourceEffectChainId, TArray<FSourceEffectChainEntry>& OutCurrentSourceEffectChainEntries) override;
diff --git a/Engine/Source/Runtime/Engine/Classes/Sound/SoundEffectBase.h b/Engine/Source/Runtime/Engine/Classes/Sound/SoundEffectBase.h
--- a/Engine/Source/Runtime/Engine/Classes/Sound/SoundEffectBase.h
+++ b/Engine/Source/Runtime/Engine/Classes/Sound/SoundEffectBase.h
@@ -78,8 +78,8 @@ public:
 	/** Enables the submix effect. */
 	void SetEnabled(const bool bInIsEnabled);
 
-	/** Updates preset on audio render thread. */
-	void Update();
+	/** Updates preset on audio render thread. Returns true if update processed a preset update, false if not. */
+	bool Update();
 
 	void SetPreset(USoundEffectPreset* Inpreset);
 
diff --git a/Engine/Source/Runtime/Engine/Classes/Sound/SoundEffectSource.h b/Engine/Source/Runtime/Engine/Classes/Sound/SoundEffectSource.h
--- a/Engine/Source/Runtime/Engine/Classes/Sound/SoundEffectSource.h
+++ b/Engine/Source/Runtime/Engine/Classes/Sound/SoundEffectSource.h
@@ -57,7 +57,7 @@ protected:
 
 };
 
-/** Struct which has data needed to initialize the source effect. */
+/** Data required to initialize the source effect. */
 struct FSoundEffectSourceInitData
 {
 	float SampleRate;
@@ -75,7 +75,7 @@ struct FSoundEffectSourceInitData
 	{}
 };
 
-/** Struct which has data to initialize the source effect. */
+/** Data required to update the source effect. */
 struct FSoundEffectSourceInputData
 {
 	float CurrentVolume;
@@ -107,5 +107,7 @@ public:
 
 	/** Process the input block of audio. Called on audio thread. */
 	virtual void ProcessAudio(const FSoundEffectSourceInputData& InData, float* OutAudioBufferData) = 0;
-};
 
+	/** Process modulation controls if enabled */
+	virtual void ProcessControls(const FSoundModulationControls& InControls) { }
+};
diff --git a/Engine/Source/Runtime/Engine/Private/SoundEffectBase.cpp b/Engine/Source/Runtime/Engine/Private/SoundEffectBase.cpp
--- a/Engine/Source/Runtime/Engine/Private/SoundEffectBase.cpp
+++ b/Engine/Source/Runtime/Engine/Private/SoundEffectBase.cpp
@@ -57,7 +57,7 @@ void FSoundEffectBase::ClearPreset()
 	}
 }
 
-void FSoundEffectBase::Update()
+bool FSoundEffectBase::Update()
 {
 	PumpPendingMessages();
 
@@ -65,7 +65,11 @@ void FSoundEffectBase::Update()
 	{
 		OnPresetChanged();
 		bChanged = false;
+
+		return true;
 	}
+
+	return false;
 }
 
 bool FSoundEffectBase::IsPreset(USoundEffectPreset* InPreset) const
diff --git a/Engine/Source/Runtime/Engine/Public/AudioDevice.h b/Engine/Source/Runtime/Engine/Public/AudioDevice.h
--- a/Engine/Source/Runtime/Engine/Public/AudioDevice.h
+++ b/Engine/Source/Runtime/Engine/Public/AudioDevice.h
@@ -1153,6 +1153,9 @@ public:
 		return bIsBakedAnalysisEnabled;
 	}
 
+	/** Updates the source's modulation controls. */
+	virtual void UpdateModulationControls(const uint32 SourceId, const FSoundModulationControls& InControls) {}
+
 	/** Updates the source effect chain. Only implemented in audio mixer. */
 	virtual void UpdateSourceEffectChain(const uint32 SourceEffectChainId, const TArray<FSourceEffectChainEntry>& SourceEffectChain, const bool bPlayEffectChainTails) {}
 
diff --git a/Engine/Source/Runtime/Engine/Public/IAudioExtensionPlugin.h b/Engine/Source/Runtime/Engine/Public/IAudioExtensionPlugin.h
--- a/Engine/Source/Runtime/Engine/Public/IAudioExtensionPlugin.h
+++ b/Engine/Source/Runtime/Engine/Public/IAudioExtensionPlugin.h
@@ -3,8 +3,10 @@
 
 #include "CoreMinimal.h"
 
+#include "AudioDefines.h"
 #include "Features/IModularFeature.h"
 #include "IAmbisonicsMixer.h"
+#include "Math/Interval.h"
 #include "Modules/ModuleInterface.h"
 #include "Modules/ModuleManager.h"
 #include "Sound/SoundSubmix.h"
@@ -456,6 +458,9 @@ public:
 	TArray<USoundModulationPluginSourceSettingsBase*> Settings;
 };
 
+/** Interface to sound that is modulateable, allowing for certain specific 
+  * behaviors to be controlled on the sound level by the modulation system.
+  */
 class ISoundModulatable
 {
 public:
@@ -520,6 +525,10 @@ public:
 };
 
 
+/*
+ * Modulatable controls found on each sound instance
+ * processed by the modulation plugin enabled
+ */
 struct FSoundModulationControls
 {
 	float Volume;
@@ -527,15 +536,126 @@ struct FSoundModulationControls
 	float Lowpass;
 	float Highpass;
 
+	TMap<FName, float> Controls;
+
 	FSoundModulationControls()
 		: Volume(1.0f)
 		, Pitch(1.0f)
-		, Lowpass(20000.0f)
-		, Highpass(20.0f)
+		, Lowpass(MAX_FILTER_FREQUENCY)
+		, Highpass(MIN_FILTER_FREQUENCY)
+	{
+	}
+};
+
+/** Parameter allowing modulation control override for systems opting in to the Modulation System. */
+USTRUCT(BlueprintType)
+struct ENGINE_API FSoundModulationParameter
+{
+	GENERATED_BODY()
+
+public:
+	/** Name of modulation control to drive parameter. Uses value last cached if control is or becomes invalid. */
+	UPROPERTY(EditAnywhere, BlueprintReadWrite, Category = Modulation)
+	FName Control;
+
+	/** Default modulation parameter value. */
+	UPROPERTY(EditAnywhere, BlueprintReadWrite, Category = Modulation)
+	float Value;
+
+private:
+	float ValueMin;
+	float ValueMax;
+
+public:
+	FSoundModulationParameter()
+		: Control(NAME_None)
+		, Value(0.0f)
+		, ValueMin(0.0f)
+		, ValueMax(1.0f)
+	{
+	}
+
+	FSoundModulationParameter(float InValue, float InValueMin = TNumericLimits<float>::Min(), float InValueMax = TNumericLimits<float>::Max())
+		: Control(NAME_None)
+		, ValueMin(InValueMin)
+		, ValueMax(InValueMax)
+	{
+		Value = FMath::Clamp(InValue, ValueMin, ValueMax);
+	}
+
+	FORCEINLINE operator float() const
+	{
+		return Value;
+	}
+
+	FORCEINLINE FSoundModulationParameter& operator= (const FSoundModulationParameter& InParam)
+	{
+		Control = InParam.Control;
+		ValueMin = InParam.ValueMin;
+		ValueMax = InParam.ValueMax;
+		Value = FMath::Clamp(InParam.Value, ValueMin, ValueMax);
+
+		return *this;
+	}
+
+	FORCEINLINE FSoundModulationParameter& operator= (float InValue)
+	{
+		Value = FMath::Clamp(InValue, ValueMin, ValueMax);
+		return *this;
+	}
+
+	FORCEINLINE float operator* (float InValue)
+	{
+		return InValue * Value;
+	}
+
+	FORCEINLINE float operator/ (float InValue)
+	{
+		return InValue / Value;
+	}
+
+	FORCEINLINE float GetValue() const
+	{
+		return Value;
+	}
+
+	FORCEINLINE float GetMinValue() const
 	{
+		return ValueMin;
+	}
+
+	FORCEINLINE float GetMaxValue() const
+	{
+		return ValueMax;
+	}
+
+	FORCEINLINE void SetBounds(float InValueMin, float InValueMax)
+	{
+		ValueMin = InValueMin;
+		ValueMax = InValueMax;
+
+		Value = FMath::Clamp(Value, InValueMin, InValueMax);
+	}
+
+	/* Returns true if updated by control, false if not */
+	bool SetValue(const FSoundModulationControls& InModControls)
+	{
+		if (Control != NAME_None)
+		{
+			if (const float* ControlValue = InModControls.Controls.Find(Control))
+			{
+				Value = FMath::Clamp(*ControlValue, ValueMin, ValueMax);
+				return true;
+			}
+		}
+
+		Value = FMath::Clamp(Value, ValueMin, ValueMax);
+		return false;
 	}
 };
 
+using FSoundModulationControlIndex = uint32;
+
 class IAudioModulation
 {
 public:
@@ -574,8 +694,8 @@ public:
 	/** Processes audio with the given input and output data structs.*/
 	virtual void ProcessAudio(const FAudioPluginSourceInputData& InputData, FAudioPluginSourceOutputData& OutputData) { }
 
-	/** Processes modulated sound controls */
-	virtual void ProcessControls(const uint32 SourceId, FSoundModulationControls& Controls) { }
+	/** Processes modulated sound controls, returning whether or not controls were modified and an update is pending. */
+	virtual bool ProcessControls(const uint32 SourceId, FSoundModulationControls& Controls) { return false; }
 
 	/** Processes all modulators */
 	virtual void ProcessModulators(const float Elapsed) { }
