Add Audio Modulation Settings Editor Improvements to beta release

diff --git a/Engine/Plugins/Runtime/AudioModulation/Source/AudioModulationEditor/AudioModulationEditor.build.cs b/Engine/Plugins/Runtime/AudioModulation/Source/AudioModulationEditor/AudioModulationEditor.build.cs
--- a/Engine/Plugins/Runtime/AudioModulation/Source/AudioModulationEditor/AudioModulationEditor.build.cs
+++ b/Engine/Plugins/Runtime/AudioModulation/Source/AudioModulationEditor/AudioModulationEditor.build.cs
@@ -18,8 +18,12 @@ public class AudioModulationEditor : ModuleRules
 				"SequenceRecorder",
 				"Slate",
 				"SlateCore",
+				"InputCore",
+				"EditorStyle",
 				"AudioEditor",
 				"AudioModulation",
+				"CurveEditor",
+				"Persona"
 			}
 		);
 
diff --git a/Engine/Plugins/Runtime/AudioModulation/Source/AudioModulationEditor/Private/AssetTypeActions/AssetTypeActions_SoundControlBus.cpp b/Engine/Plugins/Runtime/AudioModulation/Source/AudioModulationEditor/Private/AssetTypeActions/AssetTypeActions_SoundControlBus.cpp
--- a/Engine/Plugins/Runtime/AudioModulation/Source/AudioModulationEditor/Private/AssetTypeActions/AssetTypeActions_SoundControlBus.cpp
+++ b/Engine/Plugins/Runtime/AudioModulation/Source/AudioModulationEditor/Private/AssetTypeActions/AssetTypeActions_SoundControlBus.cpp
@@ -53,4 +53,14 @@ const TArray<FText>& FAssetTypeActions_SoundHPFControlBus::GetSubMenus() const
 {
 	return Modulation::SubMenus;
 }
+
+UClass* FAssetTypeActions_SoundControlBus::GetSupportedClass() const
+{
+	return USoundControlBus::StaticClass();
+}
+
+const TArray<FText>& FAssetTypeActions_SoundControlBus::GetSubMenus() const
+{
+	return Modulation::SubMenus;
+}
 #undef LOCTEXT_NAMESPACE
diff --git a/Engine/Plugins/Runtime/AudioModulation/Source/AudioModulationEditor/Private/AssetTypeActions/AssetTypeActions_SoundControlBus.h b/Engine/Plugins/Runtime/AudioModulation/Source/AudioModulationEditor/Private/AssetTypeActions/AssetTypeActions_SoundControlBus.h
--- a/Engine/Plugins/Runtime/AudioModulation/Source/AudioModulationEditor/Private/AssetTypeActions/AssetTypeActions_SoundControlBus.h
+++ b/Engine/Plugins/Runtime/AudioModulation/Source/AudioModulationEditor/Private/AssetTypeActions/AssetTypeActions_SoundControlBus.h
@@ -3,6 +3,7 @@
 
 #include "CoreMinimal.h"
 #include "AssetTypeActions_Base.h"
+#include "AudioModulationStyle.h"
 
 
 class FAssetTypeActions_SoundVolumeControlBus : public FAssetTypeActions_Base
@@ -10,7 +11,7 @@ class FAssetTypeActions_SoundVolumeControlBus : public FAssetTypeActions_Base
 public:
 	// IAssetTypeActions Implementation
 	virtual FText GetName() const override { return NSLOCTEXT("AssetTypeActions", "AssetTypeActions_SoundVolumeControlBus", "Control Bus (Volume)"); }
-	virtual FColor GetTypeColor() const override { return FColor(33, 183, 0); }
+	virtual FColor GetTypeColor() const override { return FAudioModulationStyle::GetVolumeBusColor(); }
 	virtual UClass* GetSupportedClass() const override;
 	virtual uint32 GetCategories() override { return EAssetTypeCategories::Sounds; }
 
@@ -22,7 +23,7 @@ class FAssetTypeActions_SoundPitchControlBus : public FAssetTypeActions_Base
 public:
 	// IAssetTypeActions Implementation
 	virtual FText GetName() const override { return NSLOCTEXT("AssetTypeActions", "AssetTypeActions_SoundPitchControlBus", "Control Bus (Pitch)"); }
-	virtual FColor GetTypeColor() const override { return FColor(181, 21, 0); }
+	virtual FColor GetTypeColor() const override { return FAudioModulationStyle::GetPitchBusColor(); }
 	virtual UClass* GetSupportedClass() const override;
 	virtual uint32 GetCategories() override { return EAssetTypeCategories::Sounds; }
 
@@ -34,7 +35,7 @@ class FAssetTypeActions_SoundLPFControlBus : public FAssetTypeActions_Base
 public:
 	// IAssetTypeActions Implementation
 	virtual FText GetName() const override { return NSLOCTEXT("AssetTypeActions", "AssetTypeActions_SoundLPFControlBus", "Control Bus (LPF)"); }
-	virtual FColor GetTypeColor() const override { return FColor(0, 156, 183); }
+	virtual FColor GetTypeColor() const override { return FAudioModulationStyle::GetLPFBusColor(); }
 	virtual UClass* GetSupportedClass() const override;
 	virtual uint32 GetCategories() override { return EAssetTypeCategories::Sounds; }
 
@@ -46,7 +47,19 @@ class FAssetTypeActions_SoundHPFControlBus : public FAssetTypeActions_Base
 public:
 	// IAssetTypeActions Implementation
 	virtual FText GetName() const override { return NSLOCTEXT("AssetTypeActions", "AssetTypeActions_SoundHPFControlBus", "Control Bus (HPF)"); }
-	virtual FColor GetTypeColor() const override { return FColor(94, 237, 183); }
+	virtual FColor GetTypeColor() const override { return FAudioModulationStyle::GetHPFBusColor(); }
+	virtual UClass* GetSupportedClass() const override;
+	virtual uint32 GetCategories() override { return EAssetTypeCategories::Sounds; }
+
+	virtual const TArray<FText>& GetSubMenus() const override;
+};
+
+class FAssetTypeActions_SoundControlBus : public FAssetTypeActions_Base
+{
+public:
+	// IAssetTypeActions Implementation
+	virtual FText GetName() const override { return NSLOCTEXT("AssetTypeActions", "AssetTypeActions_SoundControlBus", "Control Bus"); }
+	virtual FColor GetTypeColor() const override { return FAudioModulationStyle::GetControlBusColor(); }
 	virtual UClass* GetSupportedClass() const override;
 	virtual uint32 GetCategories() override { return EAssetTypeCategories::Sounds; }
 
diff --git a/Engine/Plugins/Runtime/AudioModulation/Source/AudioModulationEditor/Private/AssetTypeActions/AssetTypeActions_SoundModulationSettings.cpp b/Engine/Plugins/Runtime/AudioModulation/Source/AudioModulationEditor/Private/AssetTypeActions/AssetTypeActions_SoundModulationSettings.cpp
--- a/Engine/Plugins/Runtime/AudioModulation/Source/AudioModulationEditor/Private/AssetTypeActions/AssetTypeActions_SoundModulationSettings.cpp
+++ b/Engine/Plugins/Runtime/AudioModulation/Source/AudioModulationEditor/Private/AssetTypeActions/AssetTypeActions_SoundModulationSettings.cpp
@@ -1,11 +1,13 @@
 // Copyright 1998-2019 Epic Games, Inc. All Rights Reserved.
 #include "AssetTypeActions_SoundModulationSettings.h"
 
+#include "Editors/ModulationSettingsEditor.h"
 #include "SoundModulationPatch.h"
 
 
 #define LOCTEXT_NAMESPACE "AssetTypeActions"
 
+
 UClass* FAssetTypeActions_SoundModulationSettings::GetSupportedClass() const
 {
 	return USoundModulationSettings::StaticClass();
@@ -20,4 +22,18 @@ const TArray<FText>& FAssetTypeActions_SoundModulationSettings::GetSubMenus() co
 
 	return SubMenus;
 }
+
+void FAssetTypeActions_SoundModulationSettings::OpenAssetEditor(const TArray<UObject*>& InObjects, TSharedPtr<IToolkitHost> ToolkitHost)
+{
+	EToolkitMode::Type Mode = ToolkitHost.IsValid() ? EToolkitMode::WorldCentric : EToolkitMode::Standalone;
+
+	for (UObject* Object : InObjects)
+	{
+		if (USoundModulationSettings* ModulationSettings = Cast<USoundModulationSettings>(Object))
+		{
+			TSharedRef<FModulationSettingsEditor> SettingsEditor = MakeShared<FModulationSettingsEditor>();
+			SettingsEditor->Init(Mode, ToolkitHost, ModulationSettings);
+		}
+	}
+}
 #undef LOCTEXT_NAMESPACE
diff --git a/Engine/Plugins/Runtime/AudioModulation/Source/AudioModulationEditor/Private/AssetTypeActions/AssetTypeActions_SoundModulationSettings.h b/Engine/Plugins/Runtime/AudioModulation/Source/AudioModulationEditor/Private/AssetTypeActions/AssetTypeActions_SoundModulationSettings.h
--- a/Engine/Plugins/Runtime/AudioModulation/Source/AudioModulationEditor/Private/AssetTypeActions/AssetTypeActions_SoundModulationSettings.h
+++ b/Engine/Plugins/Runtime/AudioModulation/Source/AudioModulationEditor/Private/AssetTypeActions/AssetTypeActions_SoundModulationSettings.h
@@ -14,4 +14,5 @@ public:
 	virtual const TArray<FText>& GetSubMenus() const override;
 	virtual UClass* GetSupportedClass() const override;
 	virtual uint32 GetCategories() override { return EAssetTypeCategories::Sounds; }
+	virtual void OpenAssetEditor(const TArray<UObject*>& InObjects, TSharedPtr<class IToolkitHost> EditWithinLevelEditor = TSharedPtr<IToolkitHost>()) override;
 };
diff --git a/Engine/Plugins/Runtime/AudioModulation/Source/AudioModulationEditor/Private/AudioModulationEditor.cpp b/Engine/Plugins/Runtime/AudioModulation/Source/AudioModulationEditor/Private/AudioModulationEditor.cpp
--- a/Engine/Plugins/Runtime/AudioModulation/Source/AudioModulationEditor/Private/AudioModulationEditor.cpp
+++ b/Engine/Plugins/Runtime/AudioModulation/Source/AudioModulationEditor/Private/AudioModulationEditor.cpp
@@ -5,21 +5,22 @@
 #include "AssetTypeActions/AssetTypeActions_SoundControlBus.h"
 #include "AssetTypeActions/AssetTypeActions_SoundControlBusMix.h"
 #include "AssetTypeActions/AssetTypeActions_SoundModulatorLFO.h"
+#include "Editors/ModulationSettingsCurveEditorViewStacked.h"
+#include "Layouts/SoundControlModulationPatchLayout.h"
+#include "Layouts/SoundModulationParameterLayout.h"
 #include "Layouts/SoundModulationTransformLayout.h"
 #include "SoundModulationTransform.h"
-#include "Styling/SlateStyle.h"
-#include "Styling/SlateStyleRegistry.h"
+#include "ICurveEditorModule.h"
+
 
 namespace
 {
 	static const FName AssetToolsName = TEXT("AssetTools");
 
-	TSharedPtr<FSlateStyleSet> StyleSet;
-
 	template <typename T>
 	void AddAssetAction(IAssetTools& AssetTools, TArray<TSharedPtr<FAssetTypeActions_Base>>& AssetArray)
 	{
-		TSharedPtr<T> AssetAction = MakeShareable(new T);
+		TSharedPtr<T> AssetAction = MakeShared<T>();
 		TSharedPtr<FAssetTypeActions_Base> AssetActionBase = StaticCastSharedPtr<FAssetTypeActions_Base>(AssetAction);
 		AssetTools.RegisterAssetTypeActions(AssetAction.ToSharedRef());
 		AssetArray.Add(AssetActionBase);
@@ -28,7 +29,16 @@ namespace
 
 FAudioModulationEditorModule::FAudioModulationEditorModule()
 {
-	StyleSet = MakeShareable(new FSlateStyleSet("AudioModulationStyleSet"));
+}
+
+TSharedPtr<FExtensibilityManager> FAudioModulationEditorModule::GetModulationSettingsMenuExtensibilityManager()
+{
+	return ModulationSettingsMenuExtensibilityManager;
+}
+
+TSharedPtr<FExtensibilityManager> FAudioModulationEditorModule::GetModulationSettingsToolbarExtensibilityManager()
+{
+	return ModulationSettingsToolBarExtensibilityManager;
 }
 
 void FAudioModulationEditorModule::SetIcon(const FString& ClassName)
@@ -36,7 +46,7 @@ void FAudioModulationEditorModule::SetIcon(const FString& ClassName)
 	static const FVector2D Icon16 = FVector2D(16.0f, 16.0f);
 	static const FVector2D Icon64 = FVector2D(64.0f, 64.0f);
 
-	static const FString IconDir = FPaths::EngineDir() / FString::Printf(TEXT("Plugins/Runtime/AudioModulation/Icons"));
+	static const FString IconDir = FPaths::EngineDir() / FString(TEXT("Plugins/Runtime/AudioModulation/Icons"));
 
 	const FString IconFileName16 = FString::Printf(TEXT("%s_16x.png"), *ClassName);
 	const FString IconFileName64 = FString::Printf(TEXT("%s_64x.png"), *ClassName);
@@ -47,6 +57,12 @@ void FAudioModulationEditorModule::SetIcon(const FString& ClassName)
 
 void FAudioModulationEditorModule::StartupModule()
 {
+	static const FName AudioModulationStyleName(TEXT("AudioModulationStyle"));
+	StyleSet = MakeShared<FSlateStyleSet>(AudioModulationStyleName);
+
+	ModulationSettingsToolBarExtensibilityManager = MakeShared<FExtensibilityManager>();
+	ModulationSettingsMenuExtensibilityManager = MakeShared<FExtensibilityManager>();
+
 	// Register the audio editor asset type actions
 	IAssetTools& AssetTools = FModuleManager::LoadModuleChecked<FAssetToolsModule>(AssetToolsName).Get();
 
@@ -54,6 +70,7 @@ void FAudioModulationEditorModule::StartupModule()
 	AddAssetAction<FAssetTypeActions_SoundPitchControlBus>(AssetTools, AssetActions);
 	AddAssetAction<FAssetTypeActions_SoundHPFControlBus>(AssetTools, AssetActions);
 	AddAssetAction<FAssetTypeActions_SoundLPFControlBus>(AssetTools, AssetActions);
+	AddAssetAction<FAssetTypeActions_SoundControlBus>(AssetTools, AssetActions);
 	AddAssetAction<FAssetTypeActions_SoundControlBusMix>(AssetTools, AssetActions);
 	AddAssetAction<FAssetTypeActions_SoundModulatorLFO>(AssetTools, AssetActions);
 	AddAssetAction<FAssetTypeActions_SoundModulationSettings>(AssetTools, AssetActions);
@@ -62,20 +79,55 @@ void FAudioModulationEditorModule::StartupModule()
 	SetIcon(TEXT("SoundPitchControlBus"));
 	SetIcon(TEXT("SoundHPFControlBus"));
 	SetIcon(TEXT("SoundLPFControlBus"));
+	SetIcon(TEXT("SoundControlBus"));
 	SetIcon(TEXT("SoundControlBusMix"));
 	SetIcon(TEXT("SoundBusModulatorLFO"));
 	SetIcon(TEXT("SoundModulationSettings"));
 
+	RegisterCustomPropertyLayouts();
+
+	ICurveEditorModule& CurveEditorModule = FModuleManager::LoadModuleChecked<ICurveEditorModule>("CurveEditor");
+	FModCurveEditorModel::ViewId = CurveEditorModule.RegisterView(FOnCreateCurveEditorView::CreateStatic(
+		[](TWeakPtr<FCurveEditor> WeakCurveEditor) -> TSharedRef<SCurveEditorView>
+		{
+			return SNew(SModulationSettingsEditorViewStacked, WeakCurveEditor);
+		}
+	));
+
+	FSlateStyleRegistry::RegisterSlateStyle(*StyleSet.Get());
+}
+
+void FAudioModulationEditorModule::RegisterCustomPropertyLayouts()
+{
 	FPropertyEditorModule& PropertyModule = FModuleManager::LoadModuleChecked<FPropertyEditorModule>(TEXT("PropertyEditor"));
 	PropertyModule.RegisterCustomPropertyTypeLayout("SoundModulationOutputTransform",
 		FOnGetPropertyTypeCustomizationInstance::CreateStatic(
 			&FSoundModulationOutputTransformLayoutCustomization::MakeInstance));
-
-	FSlateStyleRegistry::RegisterSlateStyle(*StyleSet.Get());
+// 	PropertyModule.RegisterCustomPropertyTypeLayout("SoundModulationParameter",
+// 		FOnGetPropertyTypeCustomizationInstance::CreateStatic(
+// 			&FSoundModulationParameterLayoutCustomization::MakeInstance));
+	PropertyModule.RegisterCustomPropertyTypeLayout("SoundVolumeModulationPatch",
+		FOnGetPropertyTypeCustomizationInstance::CreateStatic(
+			&FSoundVolumeModulationPatchLayoutCustomization::MakeInstance));
+	PropertyModule.RegisterCustomPropertyTypeLayout("SoundPitchModulationPatch",
+		FOnGetPropertyTypeCustomizationInstance::CreateStatic(
+			&FSoundPitchModulationPatchLayoutCustomization::MakeInstance));
+	PropertyModule.RegisterCustomPropertyTypeLayout("SoundHPFModulationPatch",
+		FOnGetPropertyTypeCustomizationInstance::CreateStatic(
+			&FSoundHPFModulationPatchLayoutCustomization::MakeInstance));
+	PropertyModule.RegisterCustomPropertyTypeLayout("SoundLPFModulationPatch",
+		FOnGetPropertyTypeCustomizationInstance::CreateStatic(
+			&FSoundLPFModulationPatchLayoutCustomization::MakeInstance));
+	PropertyModule.RegisterCustomPropertyTypeLayout("SoundControlModulationPatch",
+		FOnGetPropertyTypeCustomizationInstance::CreateStatic(
+			&FSoundControlModulationPatchLayoutCustomization::MakeInstance));
 }
 
 void FAudioModulationEditorModule::ShutdownModule()
 {
+	ModulationSettingsToolBarExtensibilityManager.Reset();
+	ModulationSettingsMenuExtensibilityManager.Reset();
+
 	if (FModuleManager::Get().IsModuleLoaded(AssetToolsName))
 	{
 		IAssetTools& AssetTools = FModuleManager::GetModuleChecked<FAssetToolsModule>(AssetToolsName).Get();
@@ -86,8 +138,13 @@ void FAudioModulationEditorModule::ShutdownModule()
 	}
 	AssetActions.Reset();
 
+	if (ICurveEditorModule* CurveEditorModule = FModuleManager::GetModulePtr<ICurveEditorModule>("CurveEditor"))
+	{
+		CurveEditorModule->UnregisterView(FModCurveEditorModel::ViewId);
+	}
+	FModCurveEditorModel::ViewId = ECurveEditorViewID::Invalid;
+
 	FSlateStyleRegistry::UnRegisterSlateStyle(*StyleSet.Get());
 }
 
-
 IMPLEMENT_MODULE(FAudioModulationEditorModule, AudioModulationEditor);
diff --git a/Engine/Plugins/Runtime/AudioModulation/Source/AudioModulationEditor/Private/AudioModulationEditor.h b/Engine/Plugins/Runtime/AudioModulation/Source/AudioModulationEditor/Private/AudioModulationEditor.h
--- a/Engine/Plugins/Runtime/AudioModulation/Source/AudioModulationEditor/Private/AudioModulationEditor.h
+++ b/Engine/Plugins/Runtime/AudioModulation/Source/AudioModulationEditor/Private/AudioModulationEditor.h
@@ -4,7 +4,10 @@
 #include "CoreMinimal.h"
 
 #include "AssetTypeActions_Base.h"
+#include "CurveEditorTypes.h"
 #include "Modules/ModuleManager.h"
+#include "Styling/SlateStyle.h"
+#include "Styling/SlateStyleRegistry.h"
 
 
 class FAudioModulationEditorModule : public IModuleInterface
@@ -12,11 +15,22 @@ class FAudioModulationEditorModule : public IModuleInterface
 public:
 	FAudioModulationEditorModule();
 
+	TSharedPtr<FExtensibilityManager> GetModulationSettingsMenuExtensibilityManager();
+	TSharedPtr<FExtensibilityManager> GetModulationSettingsToolbarExtensibilityManager();
+
 	virtual void StartupModule() override;
+
+	void RegisterCustomPropertyLayouts();
+
 	virtual void ShutdownModule() override;
 
 private:
 	void SetIcon(const FString& ClassName);
 
 	TArray<TSharedPtr<FAssetTypeActions_Base>> AssetActions;
+
+	TSharedPtr<FExtensibilityManager> ModulationSettingsMenuExtensibilityManager;
+	TSharedPtr<FExtensibilityManager> ModulationSettingsToolBarExtensibilityManager;
+
+	TSharedPtr<FSlateStyleSet> StyleSet;
 };
diff --git a/Engine/Plugins/Runtime/AudioModulation/Source/AudioModulationEditor/Private/AudioModulationStyle.h b/Engine/Plugins/Runtime/AudioModulation/Source/AudioModulationEditor/Private/AudioModulationStyle.h
new file mode 100644
--- /dev/null
+++ b/Engine/Plugins/Runtime/AudioModulation/Source/AudioModulationEditor/Private/AudioModulationStyle.h
@@ -0,0 +1,16 @@
+// Copyright 1998-2019 Epic Games, Inc. All Rights Reserved.
+#pragma once
+
+#include "CoreMinimal.h"
+#include "Math/Color.h"
+
+
+class FAudioModulationStyle
+{
+public:
+	static const FColor GetVolumeBusColor()  { return FColor(33, 183, 0); }
+	static const FColor GetPitchBusColor()   { return FColor(181, 21, 0); }
+	static const FColor GetLPFBusColor()     { return FColor(0, 156, 183); }
+	static const FColor GetHPFBusColor()     { return FColor(94, 237, 183); }
+	static const FColor GetControlBusColor() { return FColor(215, 180, 210); }
+};
\ No newline at end of file
diff --git a/Engine/Plugins/Runtime/AudioModulation/Source/AudioModulationEditor/Private/Editors/ModulationSettingsCurveEditorViewStacked.cpp b/Engine/Plugins/Runtime/AudioModulation/Source/AudioModulationEditor/Private/Editors/ModulationSettingsCurveEditorViewStacked.cpp
new file mode 100644
--- /dev/null
+++ b/Engine/Plugins/Runtime/AudioModulation/Source/AudioModulationEditor/Private/Editors/ModulationSettingsCurveEditorViewStacked.cpp
@@ -0,0 +1,378 @@
+// Copyright 1998-2019 Epic Games, Inc. All Rights Reserved.
+
+#include "ModulationSettingsCurveEditorViewStacked.h"
+
+#include "AudioModulationStyle.h"
+#include "Curves/CurveFloat.h"
+#include "CurveEditor.h"
+#include "CurveModel.h"
+#include "EditorStyleSet.h"
+#include "Fonts/FontMeasure.h"
+#include "SCurveEditorPanel.h"
+#include "Widgets/Text/STextBlock.h"
+
+
+#define LOCTEXT_NAMESPACE "ModulationSettingsEditor"
+
+
+FModCurveEditorModel::FModCurveEditorModel(FRichCurve& InRichCurve, UObject* InOwner, FName InControlName, EModSettingsOutputEditorCurveSource InSource, UCurveFloat* InSharedCurve)
+	: FRichCurveEditorModel(&InRichCurve, InOwner)
+	, Output(EModSettingsEditorCurveOutput::Control)
+	, Source(InSource)
+{
+	SupportedViews = ViewId;
+	Refresh(EModSettingsEditorCurveOutput::Control, &InControlName, InSharedCurve);
+}
+
+FModCurveEditorModel::FModCurveEditorModel(FRichCurve& InRichCurve, UObject* InOwner, EModSettingsEditorCurveOutput InOutput, EModSettingsOutputEditorCurveSource InSource, UCurveFloat* InSharedCurve)
+	: FRichCurveEditorModel(&InRichCurve, InOwner)
+	, Output(InOutput)
+	, Source(InSource)
+{
+	SupportedViews = ViewId;
+	Refresh(InOutput, nullptr /* InControlName */, InSharedCurve);
+}
+
+void FModCurveEditorModel::Refresh(EModSettingsEditorCurveOutput InCurveOutput, const FName* InControlName, UCurveFloat* InSharedCurve)
+{
+	FText ShortNameBase;
+	switch (InCurveOutput)
+	{
+		case EModSettingsEditorCurveOutput::Volume:
+		{
+			Color = FAudioModulationStyle::GetVolumeBusColor();
+			ShortNameBase = LOCTEXT("ModulationOutputVolumeBus", "Volume");
+			IntentionName = TEXT("AudioVolume");
+		}
+		break;
+
+		case EModSettingsEditorCurveOutput::Pitch:
+		{
+			Color = FAudioModulationStyle::GetPitchBusColor();
+			ShortNameBase = LOCTEXT("ModulationOutputPitchBus", "Pitch");
+			IntentionName = TEXT("AudioPitch");
+		}
+		break;
+
+		case EModSettingsEditorCurveOutput::Highpass:
+		{
+			Color = FAudioModulationStyle::GetHPFBusColor();
+			ShortNameBase = LOCTEXT("ModulationOutputPitchBus", "Highpass");
+			IntentionName = TEXT("AudioHPFFreq");
+		}
+		break;
+
+		case EModSettingsEditorCurveOutput::Lowpass:
+		{
+			Color = FAudioModulationStyle::GetLPFBusColor();
+			ShortNameBase = LOCTEXT("ModulationOutputPitchBus", "Lowpass");
+			IntentionName = TEXT("AudioLPFFreq");
+		}
+		break;
+
+		case EModSettingsEditorCurveOutput::Control:
+		{
+			Color = FAudioModulationStyle::GetControlBusColor();
+			ShortNameBase = FText::FromName(*InControlName);
+			IntentionName = TEXT("AudioControlValue");
+		}
+		break;
+	}
+
+	switch (Source)
+	{
+		case EModSettingsOutputEditorCurveSource::Shared:
+		{
+			check(InSharedCurve);
+			FText CurveNameText = FText::FromString(InSharedCurve->GetName());
+			SetShortDisplayName(FText::Format(LOCTEXT("ModulationSharedDisplayName", "{0} (Shared - {1})"), ShortNameBase, CurveNameText));
+		}
+		break;
+
+		case EModSettingsOutputEditorCurveSource::Custom:
+		{
+			ShortDisplayName = FText::Format(LOCTEXT("ModulationOutputCurveDisplayName", "{0} (Custom)"), ShortNameBase);
+		}
+		break;
+
+		case EModSettingsOutputEditorCurveSource::Expression:
+		{
+			bKeyDrawEnabled = 0;
+			ShortDisplayName = FText::Format(LOCTEXT("ModulationOutputCurveExpressionDisplayName", "{0} (Expression)"), ShortNameBase);
+		}
+		break;
+
+		case EModSettingsOutputEditorCurveSource::Unset:
+		default:
+		{
+			bKeyDrawEnabled = 0;
+			ShortDisplayName = FText::Format(LOCTEXT("ModulationOutputCurveExpressionDisplayName", "{0} (Shared - Unset)"), ShortNameBase);
+		}
+	}
+}
+
+FLinearColor FModCurveEditorModel::GetColor() const
+{
+	return Source == EModSettingsOutputEditorCurveSource::Custom || Source == EModSettingsOutputEditorCurveSource::Expression
+		? Color
+		: Color.Desaturate(0.45f);
+}
+
+EModSettingsEditorCurveOutput FModCurveEditorModel::GetOutput() const
+{
+	return Output;
+}
+
+EModSettingsOutputEditorCurveSource FModCurveEditorModel::GetSource() const
+{
+	return Source;
+}
+
+bool FModCurveEditorModel::IsReadOnly() const
+{
+	return Source != EModSettingsOutputEditorCurveSource::Custom;
+}
+
+ECurveEditorViewID FModCurveEditorModel::ViewId = ECurveEditorViewID::Invalid;
+
+void SModulationSettingsEditorViewStacked::Construct(const FArguments& InArgs, TWeakPtr<FCurveEditor> InCurveEditor)
+{
+	SCurveEditorViewStacked::Construct(InArgs, InCurveEditor);
+
+	StackedHeight = 250.0;
+	StackedPadding = 28.0f;
+}
+
+void SModulationSettingsEditorViewStacked::PaintView(const FPaintArgs& Args, const FGeometry& AllottedGeometry, const FSlateRect& MyCullingRect, FSlateWindowElementList& OutDrawElements, int32 BaseLayerId, const FWidgetStyle& InWidgetStyle, bool bParentEnabled) const
+{
+	TSharedPtr<FCurveEditor> CurveEditor = WeakCurveEditor.Pin();
+	if (CurveEditor)
+	{
+		const ESlateDrawEffect DrawEffects = ShouldBeEnabled(bParentEnabled) ? ESlateDrawEffect::None : ESlateDrawEffect::DisabledEffect;
+
+		DrawBackground(AllottedGeometry, OutDrawElements, BaseLayerId, DrawEffects);
+		DrawViewGrids(AllottedGeometry, MyCullingRect, OutDrawElements, BaseLayerId, DrawEffects);
+		DrawLabels(AllottedGeometry, MyCullingRect, OutDrawElements, BaseLayerId, DrawEffects);
+		DrawCurves(CurveEditor.ToSharedRef(), AllottedGeometry, MyCullingRect, OutDrawElements, BaseLayerId, InWidgetStyle, DrawEffects);
+	}
+}
+
+void SModulationSettingsEditorViewStacked::DrawLabels(const FGeometry& AllottedGeometry, const FSlateRect& MyCullingRect, FSlateWindowElementList& OutDrawElements, int32 BaseLayerId, ESlateDrawEffect DrawEffects) const
+{
+	TSharedPtr<FCurveEditor> CurveEditor = WeakCurveEditor.Pin();
+	if (!CurveEditor)
+	{
+		return;
+	}
+
+	const int32 LabelLayerId = BaseLayerId + CurveViewConstants::ELayerOffset::Labels;
+
+	const double ValuePerPixel = 1.0 / StackedHeight;
+	const double ValueSpacePadding = StackedPadding * ValuePerPixel;
+
+	const FSlateFontInfo FontInfo = FCoreStyle::Get().GetFontStyle("FontAwesome.11");
+	const FVector2D LocalSize = AllottedGeometry.GetLocalSize();
+	const FCurveEditorScreenSpaceV ViewSpace = GetViewSpace();
+
+	// Draw the curve labels for each view
+	for (auto It = CurveInfoByID.CreateConstIterator(); It; ++It)
+	{
+		FCurveModel* Curve = CurveEditor->FindCurve(It.Key());
+		if (!ensureAlways(Curve))
+		{
+			continue;
+		}
+
+		const int32  CurveIndexFromBottom = CurveInfoByID.Num() - It->Value.CurveIndex - 1;
+		const double PaddingToBottomOfView = (CurveIndexFromBottom + 1)*ValueSpacePadding;
+
+		const float PixelBottom = ViewSpace.ValueToScreen(CurveIndexFromBottom + PaddingToBottomOfView);
+		const float PixelTop = ViewSpace.ValueToScreen(CurveIndexFromBottom + PaddingToBottomOfView + 1.0);
+
+		if (!FSlateRect::DoRectanglesIntersect(MyCullingRect, TransformRect(AllottedGeometry.GetAccumulatedLayoutTransform(), FSlateRect(0, PixelTop, LocalSize.X, PixelBottom))))
+		{
+			continue;
+		}
+
+		const FText Label = Curve->GetLongDisplayName();
+
+		const FVector2D Position(CurveViewConstants::CurveLabelOffsetX * 1.5f, PixelTop + CurveViewConstants::CurveLabelOffsetY * 3.0f);
+
+		const FPaintGeometry LabelGeometry = AllottedGeometry.ToPaintGeometry(FSlateLayoutTransform(Position));
+		const FPaintGeometry LabelDropshadowGeometry = AllottedGeometry.ToPaintGeometry(FSlateLayoutTransform(Position + FVector2D(2, 2)));
+
+		// Drop shadow
+		FSlateDrawElement::MakeText(OutDrawElements, LabelLayerId, LabelDropshadowGeometry, Label, FontInfo, DrawEffects, FLinearColor::Black.CopyWithNewOpacity(0.80f));
+		FSlateDrawElement::MakeText(OutDrawElements, LabelLayerId + 1, LabelGeometry, Label, FontInfo, DrawEffects, Curve->GetColor());
+	}
+}
+
+void SModulationSettingsEditorViewStacked::DrawViewGrids(const FGeometry& AllottedGeometry, const FSlateRect& MyCullingRect, FSlateWindowElementList& OutDrawElements, int32 BaseLayerId, ESlateDrawEffect DrawEffects) const
+{
+	TSharedPtr<FCurveEditor> CurveEditor = WeakCurveEditor.Pin();
+	if (!CurveEditor)
+	{
+		return;
+	}
+
+	// Rendering info
+	const float          Width = AllottedGeometry.GetLocalSize().X;
+	const float          Height = AllottedGeometry.GetLocalSize().Y;
+	const FSlateBrush*   WhiteBrush = FEditorStyle::GetBrush("WhiteBrush");
+
+	FViewGridDrawInfo DrawInfo(&AllottedGeometry, GetViewSpace(), CurveEditor->GetPanel()->GetGridLineTint(), BaseLayerId);
+
+	TArray<float> MajorGridLinesX, MinorGridLinesX;
+	TArray<FText> MajorGridLabelsX;
+
+	GetGridLinesX(CurveEditor.ToSharedRef(), MajorGridLinesX, MinorGridLinesX, &MajorGridLabelsX);
+
+	const double ValuePerPixel = 1.0 / StackedHeight;
+	const double ValueSpacePadding = StackedPadding * ValuePerPixel;
+
+	for (auto It = CurveInfoByID.CreateConstIterator(); It; ++It)
+	{
+		DrawInfo.SetCurveModel(CurveEditor->FindCurve(It.Key()));
+		if (!ensureAlways(DrawInfo.GetCurveModel()))
+		{
+			continue;
+		}
+
+		const int32  Index = CurveInfoByID.Num() - It->Value.CurveIndex - 1;
+		double Padding = (Index + 1) * ValueSpacePadding;
+		DrawInfo.SetLowerValue(Index + Padding);
+
+		FSlateRect BoundsRect(0, DrawInfo.GetPixelTop(), Width, DrawInfo.GetPixelBottom());
+		if (!FSlateRect::DoRectanglesIntersect(MyCullingRect, TransformRect(AllottedGeometry.GetAccumulatedLayoutTransform(), BoundsRect)))
+		{
+			continue;
+		}
+
+		// Tint the views based on their curve color
+		{
+			FLinearColor CurveColorTint = DrawInfo.GetCurveModel()->GetColor().CopyWithNewOpacity(0.05f);
+			const FPaintGeometry BoxGeometry = AllottedGeometry.ToPaintGeometry(
+				FVector2D(Width, StackedHeight),
+				FSlateLayoutTransform(FVector2D(0.f, DrawInfo.GetPixelTop()))
+			);
+
+			const int32 GridOverlayLayerId = DrawInfo.GetBaseLayerId() + CurveViewConstants::ELayerOffset::GridOverlays;
+			FSlateDrawElement::MakeBox(OutDrawElements, GridOverlayLayerId, BoxGeometry, WhiteBrush, DrawEffects, CurveColorTint);
+		}
+
+		// Horizontal grid lines
+		DrawInfo.LinePoints[0].X = 0.0;
+		DrawInfo.LinePoints[1].X = Width;
+
+		DrawViewGridLineX(OutDrawElements, DrawInfo, DrawEffects, 1.0  /* OffsetAlpha */, true  /* bIsMajor */);
+		DrawViewGridLineX(OutDrawElements, DrawInfo, DrawEffects, 0.75 /* OffsetAlpha */, false /* bIsMajor */);
+		DrawViewGridLineX(OutDrawElements, DrawInfo, DrawEffects, 0.5  /* OffsetAlpha */, true  /* bIsMajor */);
+		DrawViewGridLineX(OutDrawElements, DrawInfo, DrawEffects, 0.25 /* OffsetAlpha */, false /* bIsMajor */);
+		DrawViewGridLineX(OutDrawElements, DrawInfo, DrawEffects, 0.0  /* OffsetAlpha */, true  /* bIsMajor */);
+
+		// Vertical grid lines
+		{
+			DrawInfo.LinePoints[0].Y = DrawInfo.GetPixelTop();
+			DrawInfo.LinePoints[1].Y = DrawInfo.GetPixelBottom();
+
+			// Draw major vertical grid lines
+			check(MajorGridLinesX.Num() == MajorGridLabelsX.Num());
+			for (int i = 0; i < MajorGridLinesX.Num(); ++i)
+			{
+				const float VerticalLine = FMath::RoundToFloat(MajorGridLinesX[i]);
+				const FText Label = MajorGridLabelsX[i];
+
+				DrawViewGridLineY(VerticalLine, OutDrawElements, DrawInfo, DrawEffects, &Label, true /* bIsMajor */);
+			}
+
+			// Now draw the minor vertical lines which are drawn with a lighter color.
+			for (float VerticalLine : MinorGridLinesX)
+			{
+				VerticalLine = FMath::RoundToFloat(VerticalLine);
+				DrawViewGridLineY(VerticalLine, OutDrawElements, DrawInfo, DrawEffects, nullptr /* Label */, false /* bIsMajor */);
+			}
+		}
+	}
+}
+
+void SModulationSettingsEditorViewStacked::DrawViewGridLineX(FSlateWindowElementList& OutDrawElements, FViewGridDrawInfo& DrawInfo, ESlateDrawEffect DrawEffects, double OffsetAlpha, bool bIsMajor) const
+{
+	double ValueMin;
+	double ValueMax;
+	DrawInfo.GetCurveModel()->GetValueRange(ValueMin, ValueMax);
+
+	const int32 GridLineLayerId = DrawInfo.GetBaseLayerId() + CurveViewConstants::ELayerOffset::GridLines;
+
+	const double LowerValue = DrawInfo.GetLowerValue();
+	const double PixelBottom = DrawInfo.GetPixelBottom();
+	const double PixelTop = DrawInfo.GetPixelTop();
+
+	const TSharedRef<FSlateFontMeasure> FontMeasure = FSlateApplication::Get().GetRenderer()->GetFontMeasureService();
+	const FSlateFontInfo FontInfo = FCoreStyle::Get().GetFontStyle("ToolTip.LargerFont");
+
+	FLinearColor Color = bIsMajor ? DrawInfo.GetMajorGridColor() : DrawInfo.GetMinorGridColor();
+
+	DrawInfo.LinePoints[0].Y = DrawInfo.LinePoints[1].Y = DrawInfo.ScreenSpace.ValueToScreen(LowerValue + OffsetAlpha);
+	FSlateDrawElement::MakeLines(OutDrawElements, GridLineLayerId, DrawInfo.PaintGeometry, DrawInfo.LinePoints, DrawEffects, Color, false);
+
+	const FText Label = FText::AsNumber(FMath::Lerp(ValueMin, ValueMax, OffsetAlpha), &DrawInfo.LabelFormat);
+	const FVector2D LabelSize = FontMeasure->Measure(Label, FontInfo);
+	double LabelY = FMath::Lerp(PixelBottom, PixelTop, OffsetAlpha);
+	
+	// Offset label Y position below line only if the top gridline,
+	// otherwise push above
+	if (FMath::IsNearlyEqual(OffsetAlpha, 1.0))
+	{
+		LabelY += 5.0;
+	}
+	else
+	{
+		LabelY -= 15.0;
+	}
+
+	const FPaintGeometry LabelGeometry = DrawInfo.AllottedGeometry->ToPaintGeometry(FSlateLayoutTransform(FVector2D(CurveViewConstants::CurveLabelOffsetX, LabelY)));
+
+	FSlateDrawElement::MakeText(
+		OutDrawElements,
+		DrawInfo.GetBaseLayerId() + CurveViewConstants::ELayerOffset::GridLabels,
+		LabelGeometry,
+		Label,
+		FontInfo,
+		DrawEffects,
+		DrawInfo.GetLabelColor()
+	);
+}
+
+void SModulationSettingsEditorViewStacked::DrawViewGridLineY(const float VerticalLine, FSlateWindowElementList& OutDrawElements, FViewGridDrawInfo& DrawInfo, ESlateDrawEffect DrawEffects, const FText* Label, bool bIsMajor) const
+{
+	const float Width = DrawInfo.AllottedGeometry->GetLocalSize().X;
+	if (VerticalLine >= 0 || VerticalLine <= FMath::RoundToFloat(Width))
+	{
+		const int32 GridLineLayerId = DrawInfo.GetBaseLayerId() + CurveViewConstants::ELayerOffset::GridLines;
+
+		const FLinearColor LineColor = bIsMajor ? DrawInfo.GetMajorGridColor() : DrawInfo.GetMinorGridColor();
+
+		DrawInfo.LinePoints[0].X = DrawInfo.LinePoints[1].X = VerticalLine;
+		FSlateDrawElement::MakeLines(OutDrawElements, GridLineLayerId, DrawInfo.PaintGeometry, DrawInfo.LinePoints, DrawEffects, LineColor, false);
+
+		if (Label)
+		{
+			const FSlateFontInfo FontInfo = FCoreStyle::Get().GetFontStyle("ToolTip.LargerFont");
+
+			const TSharedRef<FSlateFontMeasure> FontMeasure = FSlateApplication::Get().GetRenderer()->GetFontMeasureService();
+			const FVector2D LabelSize = FontMeasure->Measure(*Label, FontInfo);
+			const FPaintGeometry LabelGeometry = DrawInfo.AllottedGeometry->ToPaintGeometry(FSlateLayoutTransform(FVector2D(VerticalLine, DrawInfo.LinePoints[0].Y - LabelSize.Y * 1.2f)));
+
+			FSlateDrawElement::MakeText(
+				OutDrawElements,
+				DrawInfo.GetBaseLayerId() + CurveViewConstants::ELayerOffset::GridLabels,
+				LabelGeometry,
+				*Label,
+				FontInfo,
+				DrawEffects,
+				DrawInfo.GetLabelColor()
+			);
+		}
+	}
+}
+#undef LOCTEXT_NAMESPACE
\ No newline at end of file
diff --git a/Engine/Plugins/Runtime/AudioModulation/Source/AudioModulationEditor/Private/Editors/ModulationSettingsCurveEditorViewStacked.h b/Engine/Plugins/Runtime/AudioModulation/Source/AudioModulationEditor/Private/Editors/ModulationSettingsCurveEditorViewStacked.h
new file mode 100644
--- /dev/null
+++ b/Engine/Plugins/Runtime/AudioModulation/Source/AudioModulationEditor/Private/Editors/ModulationSettingsCurveEditorViewStacked.h
@@ -0,0 +1,185 @@
+// Copyright 1998-2019 Epic Games, Inc. All Rights Reserved.
+#pragma once
+
+#include "CurveModel.h"
+#include "Curves/RichCurve.h"
+#include "RichCurveEditorModel.h"
+#include "UObject/Object.h"
+#include "Views/SCurveEditorViewStacked.h"
+#include "Views/SInteractiveCurveEditorView.h"
+
+// Forward Declarations
+class UCurveFloat;
+
+
+enum class EModSettingsEditorCurveOutput : uint8
+{
+	Volume,
+	Pitch,
+	Highpass,
+	Lowpass,
+	Control,
+	Count
+};
+
+
+enum class EModSettingsOutputEditorCurveSource : uint8
+{
+	Custom,
+	Expression,
+	Shared,
+	Unset
+};
+
+
+struct FViewGridDrawInfo
+{
+	const FGeometry* AllottedGeometry;
+
+	FPaintGeometry PaintGeometry;
+
+	ESlateDrawEffect DrawEffects;
+	FNumberFormattingOptions LabelFormat;
+
+	TArray<FVector2D> LinePoints;
+
+	FCurveEditorScreenSpace ScreenSpace;
+
+private:
+	FLinearColor MajorGridColor;
+	FLinearColor MinorGridColor;
+
+	int32 BaseLayerId;
+
+	const FCurveModel* CurveModel;
+
+	double LowerValue;
+	double PixelBottom;
+	double PixelTop;
+
+public:
+	FViewGridDrawInfo(const FGeometry* InAllottedGeometry, const FCurveEditorScreenSpace& InScreenSpace, FLinearColor InGridColor, int32 InBaseLayerId)
+		: AllottedGeometry(InAllottedGeometry)
+		, ScreenSpace(InScreenSpace)
+		, BaseLayerId(InBaseLayerId)
+		, CurveModel(nullptr)
+		, LowerValue(0)
+		, PixelBottom(0)
+		, PixelTop(0)
+	{
+		check(AllottedGeometry);
+
+		// Pre-allocate an array of line points to draw our vertical lines. Each major grid line
+		// will overwrite the X value of both points but leave the Y value untouched so they draw from the bottom to the top.
+		LinePoints.Add(FVector2D(0.f, 0.f));
+		LinePoints.Add(FVector2D(0.f, 0.f));
+
+		MajorGridColor = InGridColor;
+		MinorGridColor = InGridColor.CopyWithNewOpacity(InGridColor.A * 0.5f);
+
+		PaintGeometry = InAllottedGeometry->ToPaintGeometry();
+
+		LabelFormat.SetMaximumFractionalDigits(6);
+	}
+
+	void SetCurveModel(const FCurveModel * InCurveModel)
+	{
+		CurveModel = InCurveModel;
+	}
+
+	const FCurveModel* GetCurveModel() const
+	{
+		return CurveModel;
+	}
+
+	void SetLowerValue(double InLowerValue)
+	{
+		LowerValue = InLowerValue;
+		PixelBottom = ScreenSpace.ValueToScreen(InLowerValue);
+		PixelTop = ScreenSpace.ValueToScreen(InLowerValue + 1.0);
+	}
+
+	int32 GetBaseLayerId() const
+	{
+		return BaseLayerId;
+	}
+
+	FLinearColor GetLabelColor() const
+	{
+		check(CurveModel);
+		return GetCurveModel()->GetColor().CopyWithNewOpacity(0.7f);
+	}
+
+	double GetLowerValue() const
+	{
+		return LowerValue;
+	}
+
+	FLinearColor GetMajorGridColor() const
+	{
+		return MajorGridColor;
+	}
+
+	FLinearColor GetMinorGridColor() const
+	{
+		return MinorGridColor;
+	}
+
+	double GetPixelBottom() const
+	{
+		return PixelBottom;
+	}
+
+	double GetPixelTop() const
+	{
+		return PixelTop;
+	}
+};
+
+
+class FModCurveEditorModel : public FRichCurveEditorModel
+{
+public:
+
+	static ECurveEditorViewID ViewId;
+
+	FModCurveEditorModel(FRichCurve& InRichCurve, UObject* InOwner, FName InControlName, EModSettingsOutputEditorCurveSource InSource, UCurveFloat* SharedCurve);
+	FModCurveEditorModel(FRichCurve& InRichCurve, UObject* InOwner, EModSettingsEditorCurveOutput InOutput, EModSettingsOutputEditorCurveSource InSource, UCurveFloat* SharedCurve);
+
+	virtual bool IsReadOnly() const override;
+
+	virtual FLinearColor GetColor() const override;
+	EModSettingsEditorCurveOutput GetOutput() const;
+	EModSettingsOutputEditorCurveSource GetSource() const;
+
+	void Refresh(EModSettingsEditorCurveOutput InCurveOutput, const FName* InControlName, UCurveFloat* InSharedCurve);
+
+private:
+
+	EModSettingsEditorCurveOutput Output;
+	EModSettingsOutputEditorCurveSource Source;
+};
+
+class SModulationSettingsEditorViewStacked : public SCurveEditorViewStacked
+{
+public:
+
+	void Construct(const FArguments& InArgs, TWeakPtr<FCurveEditor> InCurveEditor);
+
+protected:
+	virtual void PaintView(
+		const FPaintArgs&			Args,
+		const FGeometry&			AllottedGeometry,
+		const FSlateRect&			MyCullingRect,
+		FSlateWindowElementList&	OutDrawElements,
+		int32						BaseLayerId, 
+		const FWidgetStyle&			InWidgetStyle,
+		bool						bParentEnabled) const override;
+
+	virtual void DrawViewGrids(const FGeometry& AllottedGeometry, const FSlateRect& MyCullingRect, FSlateWindowElementList& OutDrawElements, int32 BaseLayerId, ESlateDrawEffect DrawEffects) const override;
+	virtual void DrawLabels(const FGeometry& AllottedGeometry, const FSlateRect& MyCullingRect, FSlateWindowElementList& OutDrawElements, int32 BaseLayerId, ESlateDrawEffect DrawEffects) const override;
+
+private:
+	void DrawViewGridLineX(FSlateWindowElementList& OutDrawElements, FViewGridDrawInfo& DrawInfo, ESlateDrawEffect DrawEffect, double OffsetAlpha, bool bIsMajor) const;
+	void DrawViewGridLineY(const float VerticalLine, FSlateWindowElementList& OutDrawElements, FViewGridDrawInfo &DrawInfo, ESlateDrawEffect DrawEffects, const FText* Label, bool bIsMajor) const;
+};
diff --git a/Engine/Plugins/Runtime/AudioModulation/Source/AudioModulationEditor/Private/Editors/ModulationSettingsEditor.cpp b/Engine/Plugins/Runtime/AudioModulation/Source/AudioModulationEditor/Private/Editors/ModulationSettingsEditor.cpp
new file mode 100644
--- /dev/null
+++ b/Engine/Plugins/Runtime/AudioModulation/Source/AudioModulationEditor/Private/Editors/ModulationSettingsEditor.cpp
@@ -0,0 +1,573 @@
+// Copyright 1998-2019 Epic Games, Inc. All Rights Reserved.
+#include "Editors/ModulationSettingsEditor.h"
+
+#include "AudioModulationEditor.h"
+#include "AudioModulationStyle.h"
+#include "CommonFrameRates.h"
+#include "Containers/Set.h"
+#include "Curves/CurveBase.h"
+#include "Curves/CurveLinearColor.h"
+#include "CurveEditor.h"
+#include "CurveEditorCommands.h"
+#include "EditorStyleSet.h"
+#include "Framework/Commands/UIAction.h"
+#include "Framework/MultiBox/MultiBoxDefs.h"
+#include "Framework/MultiBox/MultiBoxBuilder.h"
+#include "Framework/MultiBox/MultiBoxExtender.h"
+#include "IDetailsView.h"
+#include "ModulationSettingsEditor.h"
+#include "Modules/ModuleManager.h"
+#include "PropertyEditorModule.h"
+#include "RichCurveEditorModel.h"
+#include "SCurveEditorPanel.h"
+#include "SoundModulationPatch.h"
+#include "Tree/SCurveEditorTree.h"
+#include "Tree/ICurveEditorTreeItem.h"
+#include "Tree/SCurveEditorTreePin.h"
+#include "Widgets/Docking/SDockTab.h"
+#include "Widgets/Input/SButton.h"
+#include "Widgets/Input/SNumericDropDown.h"
+#include "Widgets/Layout/SBorder.h"
+#include "Widgets/SFrameRatePicker.h"
+#include "ModulationSettingsCurveEditorViewStacked.h"
+
+
+#define LOCTEXT_NAMESPACE "ModulationSettingsEditor"
+
+
+const FName FModulationSettingsEditor::AppIdentifier(TEXT("ModulationSettingsEditorApp"));
+const FName FModulationSettingsEditor::CurveTabId(TEXT("ModulationSettingsEditor_Curves"));
+const FName FModulationSettingsEditor::PropertiesTabId(TEXT("ModulationSettingsEditor_Properties"));
+
+FModulationSettingsEditor::FModulationSettingsEditor()
+{
+}
+
+void FModulationSettingsEditor::RegisterTabSpawners(const TSharedRef<FTabManager>& InTabManager)
+{
+	WorkspaceMenuCategory = InTabManager->AddLocalWorkspaceMenuCategory(LOCTEXT("WorkspaceMenu_ModulationSettingsEditor", "Modulation Settings Editor"));
+
+	FAssetEditorToolkit::RegisterTabSpawners(InTabManager);
+
+	InTabManager->RegisterTabSpawner(PropertiesTabId, FOnSpawnTab::CreateSP(this, &FModulationSettingsEditor::SpawnTab_Properties))
+		.SetDisplayName(LOCTEXT("DetailsTab", "Details"))
+		.SetGroup(WorkspaceMenuCategory.ToSharedRef())
+		.SetIcon(FSlateIcon(FEditorStyle::GetStyleSetName(), "LevelEditor.Tabs.Details"));
+
+	FSlateIcon CurveIcon(FEditorStyle::GetStyleSetName(), "ClassIcon.CurveBase");
+
+	InTabManager->RegisterTabSpawner(CurveTabId, FOnSpawnTab::CreateLambda([this](const FSpawnTabArgs& Args) { return SpawnTab_OutputCurve(Args); }))
+		.SetDisplayName(LOCTEXT("TransformCurvesTab", "Transform Curves"))
+		.SetGroup(WorkspaceMenuCategory.ToSharedRef())
+		.SetIcon(CurveIcon);
+}
+
+void FModulationSettingsEditor::UnregisterTabSpawners(const TSharedRef<FTabManager>& InTabManager)
+{
+	InTabManager->UnregisterTabSpawner(PropertiesTabId);
+	InTabManager->UnregisterTabSpawner(CurveTabId);
+}
+
+void FModulationSettingsEditor::Init(const EToolkitMode::Type Mode, const TSharedPtr<IToolkitHost>& InitToolkitHost, USoundModulationSettings* InSettingsToEdit)
+{
+	CurveEditor = MakeShared<FCurveEditor>();
+	FCurveEditorInitParams InitParams;
+	CurveEditor->InitCurveEditor(InitParams);
+	CurveEditor->GridLineLabelFormatXAttribute = LOCTEXT("GridXLabelFormat", "{0}");
+
+	TUniquePtr<ICurveEditorBounds> EditorBounds = MakeUnique<FStaticCurveEditorBounds>();
+	EditorBounds->SetInputBounds(0.05, 1.05);
+	CurveEditor->SetBounds(MoveTemp(EditorBounds));
+
+	CurvePanel = SNew(SCurveEditorPanel, CurveEditor.ToSharedRef());
+
+	// Support undo/redo
+	InSettingsToEdit->SetFlags(RF_Transactional);
+	GEditor->RegisterForUndo(this);
+
+	FDetailsViewArgs Args;
+	Args.bHideSelectionTip = true;
+	Args.NotifyHook = this;
+
+	FPropertyEditorModule& PropertyModule = FModuleManager::LoadModuleChecked<FPropertyEditorModule>("PropertyEditor");
+	PropertiesView = PropertyModule.CreateDetailView(Args);
+	PropertiesView->SetObject(InSettingsToEdit);
+
+	TSharedRef<FTabManager::FLayout> StandaloneDefaultLayout = FTabManager::NewLayout("Standalone_ModulationSettingsEditor_Layout_v1")
+		->AddArea
+		(
+			FTabManager::NewPrimaryArea()
+			->SetOrientation(Orient_Vertical)
+			->Split
+			(
+				FTabManager::NewStack()
+				->SetSizeCoefficient(0.1f)
+				->SetHideTabWell(true)
+				->AddTab(GetToolbarTabId(), ETabState::OpenedTab)
+			)
+			->Split
+			(
+				FTabManager::NewSplitter()
+				->SetSizeCoefficient(0.9f)
+				->SetOrientation(Orient_Horizontal)
+				->Split
+				(
+					FTabManager::NewStack()
+					->SetSizeCoefficient(0.225f)
+					->AddTab(PropertiesTabId, ETabState::OpenedTab)
+				)
+				->Split
+				(
+					FTabManager::NewSplitter()
+					->SetSizeCoefficient(0.775f)
+					->SetOrientation(Orient_Vertical)
+					->Split
+					(
+						FTabManager::NewStack()
+						->SetHideTabWell(true)
+						->SetSizeCoefficient(0.33f)
+						->AddTab(CurveTabId, ETabState::OpenedTab)
+					)
+				)
+			)
+		);
+
+	const bool bCreateDefaultStandaloneMenu = true;
+	const bool bCreateDefaultToolbar = true;
+	const bool bToolbarFocusable = false;
+	const bool bUseSmallIcons = true;
+	FAssetEditorToolkit::InitAssetEditor(
+		Mode,
+		InitToolkitHost,
+		AppIdentifier,
+		StandaloneDefaultLayout,
+		bCreateDefaultStandaloneMenu,
+		bCreateDefaultToolbar,
+		InSettingsToEdit,
+		bToolbarFocusable,
+		bUseSmallIcons);
+
+	FAudioModulationEditorModule& AudioEditorModule = FModuleManager::LoadModuleChecked<FAudioModulationEditorModule>(TEXT("AudioModulationEditor"));
+	AddMenuExtender(AudioEditorModule.GetModulationSettingsMenuExtensibilityManager()->GetAllExtenders(GetToolkitCommands(), GetEditingObjects()));
+	AddToolbarExtender(AudioEditorModule.GetModulationSettingsToolbarExtensibilityManager()->GetAllExtenders(GetToolkitCommands(), GetEditingObjects()));
+	AddToolbarExtender(CurvePanel->GetToolbarExtender());
+
+	if (CurveEditor.IsValid())
+	{
+		RegenerateMenusAndToolbars();
+	}
+}
+
+void FModulationSettingsEditor::GenerateExpressionCurve(const FSoundModulationOutputBase& InOutput, int32 InCurveIndex, bool bIsUnset)
+{
+	if (!CurveEditor.IsValid())
+	{
+		return;
+	}
+
+	ExpressionCurves.SetNum(InCurveIndex + 1, false /* bAllowShrinking */);
+	TSharedPtr<FRichCurve> NewCurve = MakeShared<FRichCurve>();
+	ExpressionCurves[InCurveIndex] = NewCurve;
+
+	int32 CurveResolution;
+	switch (InOutput.Transform.Curve)
+	{
+		case ESoundModulatorOutputCurve::Linear:
+		{
+			CurveResolution = 2;
+		}
+		break;
+
+		case ESoundModulatorOutputCurve::Sin:
+		case ESoundModulatorOutputCurve::SCurve:
+		{
+			CurveResolution = 64;
+		}
+		break;
+
+		case ESoundModulatorOutputCurve::Log:
+		case ESoundModulatorOutputCurve::Exp:
+		case ESoundModulatorOutputCurve::Exp_Inverse:
+		{
+			CurveResolution = 256;
+		}
+		break;
+
+		default:
+		{
+			CurveResolution = 128;
+		}
+		break;
+	}
+
+	const float CurveResolutionRatio = 1.0f / (CurveResolution - 1);
+	for (int32 i = 0; i < CurveResolution; ++i)
+	{
+		const float X = FMath::Lerp(InOutput.Transform.InputMin, InOutput.Transform.InputMax, CurveResolutionRatio * i);
+		float Y = X;
+		InOutput.Transform.Apply(Y);
+
+		NewCurve->AddKey(X, Y);
+	}
+
+	const EModSettingsOutputEditorCurveSource Source = bIsUnset ? EModSettingsOutputEditorCurveSource::Unset : EModSettingsOutputEditorCurveSource::Expression;
+	SetCurveAtOrderIndex(InCurveIndex, *NewCurve.Get(), Source, nullptr);
+}
+
+void FModulationSettingsEditor::SetCurveAtOrderIndex(int32 InCurveIndex, FRichCurve& InRichCurve, EModSettingsOutputEditorCurveSource InSource, UCurveFloat* InSharedCurve)
+{
+	check(CurveEditor.IsValid());
+
+	bool bIsNew = false;
+	while (InCurveIndex >= CurveModels.Num())
+	{
+		CurveModels.Add(FCurveModelID::Unique());
+		bIsNew = true;
+	}
+
+	// Find if the incoming rich curve is already set.  If so, early out
+	if (!bIsNew)
+	{
+		FCurveModel* CurrentModel = CurveEditor->FindCurve(CurveModels[InCurveIndex]);
+		if (CurrentModel && CurrentModel->GetCurve() == &InRichCurve)
+		{
+			FModCurveEditorModel* ModModel = static_cast<FModCurveEditorModel*>(CurrentModel);
+			if (InCurveIndex < static_cast<int32>(EModSettingsEditorCurveOutput::Control))
+			{
+				ModModel->Refresh(static_cast<EModSettingsEditorCurveOutput>(InCurveIndex), nullptr, InSharedCurve);
+			}
+			else
+			{
+				check(!InSharedCurve);
+
+				FName ControlName = GetControlName(InCurveIndex);
+				ModModel->Refresh(EModSettingsEditorCurveOutput::Control, &ControlName, InSharedCurve);
+			}
+			CurveEditor->PinCurve(CurveModels[InCurveIndex]);
+			return;
+		}
+	}
+
+	// Curves don't match, so remove what's already there
+	if (FCurveModel* CurveModel = CurveEditor->FindCurve(CurveModels[InCurveIndex]))
+	{
+		CurveEditor->RemoveCurve(CurveModels[InCurveIndex]);
+	}
+
+	// Finally, create a FCurveModel add the new FRichCurve
+	if (InCurveIndex < static_cast<int32>(EModSettingsEditorCurveOutput::Control))
+	{
+		const EModSettingsEditorCurveOutput CurveOutput = static_cast<EModSettingsEditorCurveOutput>(InCurveIndex);
+		TUniquePtr<FModCurveEditorModel> NewCurve = MakeUnique<FModCurveEditorModel>(InRichCurve, GetEditingObject(), CurveOutput, InSource, InSharedCurve);
+		CurveModels[InCurveIndex] = CurveEditor->AddCurve(MoveTemp(NewCurve));
+	}
+	else
+	{
+		const FName Name = CastChecked<USoundModulationSettings>(GetEditingObject())->Controls[GetControlIndex(InCurveIndex)].Control;
+		TUniquePtr<FModCurveEditorModel> NewCurve = MakeUnique<FModCurveEditorModel>(InRichCurve, GetEditingObject(), Name, InSource, InSharedCurve);
+		CurveModels[InCurveIndex] = CurveEditor->AddCurve(MoveTemp(NewCurve));
+	}
+
+	CurveEditor->PinCurve(CurveModels[InCurveIndex]);
+}
+
+int32 FModulationSettingsEditor::GetControlModulationCount() const
+{
+	USoundModulationSettings* Settings = Cast<USoundModulationSettings>(GetEditingObject());
+	if (!Settings)
+	{
+		return 0;
+	}
+
+	return Settings->Controls.Num();
+}
+
+int32 FModulationSettingsEditor::GetCurveCount() const
+{
+	USoundModulationSettings* Settings = Cast<USoundModulationSettings>(GetEditingObject());
+	if (!Settings)
+	{
+		return 0;
+	}
+
+	return static_cast<int32>(EModSettingsEditorCurveOutput::Count) - 1 + Settings->Controls.Num();
+}
+
+int32 FModulationSettingsEditor::GetCurveOrderIndex(EModSettingsEditorCurveOutput InCurveOutput, int32 InControlIndex) const
+{
+	switch (InCurveOutput)
+	{
+	case EModSettingsEditorCurveOutput::Volume:
+	case EModSettingsEditorCurveOutput::Pitch:
+	case EModSettingsEditorCurveOutput::Highpass:
+	case EModSettingsEditorCurveOutput::Lowpass:
+	{
+		return static_cast<int32>(InCurveOutput);
+	}
+	break;
+
+	case EModSettingsEditorCurveOutput::Control:
+	{
+		return static_cast<int32>(InCurveOutput) + InControlIndex;
+	}
+	break;
+
+	default:
+		static_assert(static_cast<int32>(EModSettingsEditorCurveOutput::Count) == 5, "Possible missing case coverage for enum EModSettingsEditorCurveOutput");
+		return -1;
+	}
+}
+
+int32 FModulationSettingsEditor::GetControlIndex(int32 InCurveIndex) const
+{
+	const int32 ControlIndexOffset = static_cast<int32>(EModSettingsEditorCurveOutput::Control);
+	if (InCurveIndex < ControlIndexOffset)
+	{
+		return -1;
+	}
+
+	return InCurveIndex - ControlIndexOffset;
+}
+
+FName FModulationSettingsEditor::GetControlName(int32 InCurveIndex) const
+{
+	const int32 ControlIndexOffset = static_cast<int32>(EModSettingsEditorCurveOutput::Control);
+	if (InCurveIndex < ControlIndexOffset)
+	{
+		return NAME_None;
+	}
+
+	USoundModulationSettings* Settings = Cast<USoundModulationSettings>(GetEditingObject());
+	if (!Settings)
+	{
+		return NAME_None;
+	}
+
+	const int32 ControlIndex = InCurveIndex - static_cast<int32>(EModSettingsEditorCurveOutput::Control);
+	return Settings->Controls[ControlIndex].Control;
+}
+
+FSoundModulationOutputBase* FModulationSettingsEditor::FindModulationOutput(int32 InCurveIndex)
+{
+	USoundModulationSettings* Settings = Cast<USoundModulationSettings>(GetEditingObject());
+	if (!Settings)
+	{
+		return nullptr;
+	}
+
+	if (InCurveIndex >= static_cast<int32>(EModSettingsEditorCurveOutput::Control))
+	{
+		const int32 ControlIndex = InCurveIndex - static_cast<int32>(EModSettingsEditorCurveOutput::Control);
+		return &Settings->Controls[ControlIndex].Output;
+	}
+
+	const EModSettingsEditorCurveOutput CurveOutput = static_cast<EModSettingsEditorCurveOutput>(InCurveIndex);
+	switch (CurveOutput)
+	{
+		case EModSettingsEditorCurveOutput::Volume:
+		{
+			return static_cast<FSoundModulationOutputBase*>(&Settings->Volume.Output);
+		}
+		break;
+
+		case EModSettingsEditorCurveOutput::Pitch:
+		{
+			return static_cast<FSoundModulationOutputBase*>(&Settings->Pitch.Output);
+		}
+		break;
+
+		case EModSettingsEditorCurveOutput::Highpass:
+		{
+			return static_cast<FSoundModulationOutputBase*>(&Settings->Highpass.Output);
+		}
+		break;
+
+		case EModSettingsEditorCurveOutput::Lowpass:
+		{
+			return static_cast<FSoundModulationOutputBase*>(&Settings->Lowpass.Output);
+		}
+		break;
+
+		default:
+		{
+			static_assert(static_cast<int32>(EModSettingsEditorCurveOutput::Count) == 5, "Possible missing case coverage for enum EModSettingsEditorCurveOutput");
+			return nullptr;
+		}
+		break;
+	}
+}
+
+FName FModulationSettingsEditor::GetToolkitFName() const
+{
+	return FName("ModulationSettingsEditor");
+}
+
+FText FModulationSettingsEditor::GetBaseToolkitName() const
+{
+	return LOCTEXT( "AppLabel", "Modulation Settings Editor" );
+}
+
+FString FModulationSettingsEditor::GetWorldCentricTabPrefix() const
+{
+	return LOCTEXT("WorldCentricTabPrefix", "ModulationSettings ").ToString();
+}
+
+FLinearColor FModulationSettingsEditor::GetWorldCentricTabColorScale() const
+{
+	return FLinearColor( 0.0f, 0.0f, 0.2f, 0.5f );
+}
+
+EOrientation FModulationSettingsEditor::GetSnapLabelOrientation() const
+{
+	return FMultiBoxSettings::UseSmallToolBarIcons.Get()
+		? EOrientation::Orient_Horizontal
+		: EOrientation::Orient_Vertical;
+}
+
+void FModulationSettingsEditor::NotifyPostChange(const FPropertyChangedEvent& PropertyChangedEvent, UProperty* PropertyThatChanged)
+{
+	if (PropertyChangedEvent.ChangeType != EPropertyChangeType::Interactive)
+	{
+		UpdateCurves();
+	}
+}
+
+TSharedRef<SDockTab> FModulationSettingsEditor::SpawnTab_Properties(const FSpawnTabArgs& Args)
+{
+	check(Args.GetTabId() == PropertiesTabId);
+
+	return SNew(SDockTab)
+		.Icon(FEditorStyle::GetBrush("LevelEditor.Tabs.Details"))
+		.Label(LOCTEXT("SoundModulationSettingsDetailsTitle", "Details"))
+		[
+			PropertiesView.ToSharedRef()
+		];
+}
+
+TSharedRef<SDockTab> FModulationSettingsEditor::SpawnTab_OutputCurve(const FSpawnTabArgs& Args)
+{
+	UpdateCurves();
+	CurveEditor->ZoomToFit();
+
+	TSharedRef<SDockTab> NewDockTab = SNew(SDockTab)
+		.Icon(FEditorStyle::GetBrush("ModulationSettingsEditor.Tabs.Properties"))
+		.Label(FText::Format(LOCTEXT("ModulationSettingsFilterTitle", "Filter Transform Curve: {0}"), FText::FromString(GetEditingObject()->GetName())))
+		.TabColorScale(GetTabColorScale())
+		[
+			SNew(SBorder)
+			.BorderImage(FEditorStyle::GetBrush("ToolPanel.GroupBorder"))
+		.Padding(0.0f)
+		[
+			CurvePanel.ToSharedRef()
+		]
+		];
+
+		return NewDockTab;
+}
+
+void FModulationSettingsEditor::PostUndo(bool bSuccess)
+{
+	if (bSuccess)
+	{
+		UpdateCurves();
+	}
+
+	CastChecked<USoundModulationSettings>(GetEditingObject())->OnPostEditChange(nullptr);
+}
+
+void FModulationSettingsEditor::UpdateCurve(int32 InCurveIndex)
+{
+	check(CurveEditor.IsValid());
+
+	FSoundModulationOutputBase* Output = FindModulationOutput(InCurveIndex);
+	if (!Output)
+	{
+		return;
+	}
+
+	switch (Output->Transform.Curve)
+	{
+		case ESoundModulatorOutputCurve::Exp:
+		case ESoundModulatorOutputCurve::Exp_Inverse:
+		case ESoundModulatorOutputCurve::Linear:
+		case ESoundModulatorOutputCurve::Log:
+		case ESoundModulatorOutputCurve::SCurve:
+		case ESoundModulatorOutputCurve::Sin:
+		{
+			GenerateExpressionCurve(*Output, InCurveIndex);
+		}
+		break;
+
+		case ESoundModulatorOutputCurve::Shared:
+		{
+			if (UCurveFloat* SharedCurve = Output->Transform.CurveShared)
+			{
+				SetCurveAtOrderIndex(InCurveIndex, SharedCurve->FloatCurve, EModSettingsOutputEditorCurveSource::Shared, SharedCurve);
+			}
+			else
+			{
+				// Builds a dummy expression that just maps input to output in case
+				// where asset isn't selected and leave source as unset
+				GenerateExpressionCurve(*Output, InCurveIndex, true /* bIsUnset */);
+			}
+		}
+		break;
+
+		case ESoundModulatorOutputCurve::Custom:
+		{
+			TrimKeys(*Output);
+			SetCurveAtOrderIndex(InCurveIndex, Output->Transform.CurveCustom, EModSettingsOutputEditorCurveSource::Custom);
+		}
+		break;
+
+		default:
+		{
+			static_assert(static_cast<int32>(ESoundModulatorOutputCurve::Count) == 8, "Possible missing case coverage for output curve.");
+		}
+		break;
+	}
+}
+
+void FModulationSettingsEditor::UpdateCurves()
+{
+	if (!CurveEditor.IsValid())
+	{
+		return;
+	}
+
+	for (int32 i = 0; i < GetCurveCount(); ++i)
+	{
+		UpdateCurve(i);
+	}
+
+	// Remove old curves from editor
+	for (int32 i = GetCurveCount(); i < CurveModels.Num(); ++i)
+	{
+		CurveEditor->RemoveCurve(CurveModels[i]);
+	}
+}
+
+void FModulationSettingsEditor::PostRedo(bool bSuccess)
+{
+	if (bSuccess)
+	{
+		UpdateCurves();
+	}
+
+	CastChecked<USoundModulationSettings>(GetEditingObject())->OnPostEditChange(nullptr);
+}
+
+void FModulationSettingsEditor::TrimKeys(FSoundModulationOutputBase& InOutput) const
+{
+	FRichCurve& Curve = InOutput.Transform.CurveCustom;
+	while (Curve.GetNumKeys() > 0 && InOutput.Transform.InputMin > Curve.GetFirstKey().Time)
+	{
+		Curve.DeleteKey(Curve.GetFirstKeyHandle());
+	}
+
+	while (Curve.GetNumKeys() > 0 && InOutput.Transform.InputMax < Curve.GetLastKey().Time)
+	{
+		Curve.DeleteKey(Curve.GetLastKeyHandle());
+	}
+}
+#undef LOCTEXT_NAMESPACE
diff --git a/Engine/Plugins/Runtime/AudioModulation/Source/AudioModulationEditor/Private/Editors/ModulationSettingsEditor.h b/Engine/Plugins/Runtime/AudioModulation/Source/AudioModulationEditor/Private/Editors/ModulationSettingsEditor.h
new file mode 100644
--- /dev/null
+++ b/Engine/Plugins/Runtime/AudioModulation/Source/AudioModulationEditor/Private/Editors/ModulationSettingsEditor.h
@@ -0,0 +1,103 @@
+// Copyright 1998-2019 Epic Games, Inc. All Rights Reserved.
+
+#pragma once
+
+#include "CoreMinimal.h"
+
+#include "CurveEditorTypes.h"
+#include "Curves/CurveOwnerInterface.h"
+#include "EditorUndoClient.h"
+#include "Framework/Docking/TabManager.h"
+#include "Math/Color.h"
+#include "Misc/NotifyHook.h"
+#include "ModulationSettingsCurveEditorViewStacked.h"
+#include "SoundModulationPatch.h"
+#include "Toolkits/AssetEditorToolkit.h"
+#include "Toolkits/IToolkitHost.h"
+#include "Widgets/SWidget.h"
+
+
+// Forward Declarations
+class FCurveEditor;
+class IToolkitHost;
+class SCurveEditorPanel;
+class UCurveBase;
+class USoundModulationSettings;
+
+
+class FModulationSettingsEditor : public FAssetEditorToolkit, public FNotifyHook, public FEditorUndoClient
+{
+public:
+	FModulationSettingsEditor();
+
+	virtual ~FModulationSettingsEditor() = default;
+
+	void Init(const EToolkitMode::Type Mode, const TSharedPtr<IToolkitHost>& InitToolkitHost, USoundModulationSettings* SettingsToEdit);
+
+	/** FAssetEditorToolkit interface */
+	virtual FName GetToolkitFName() const override;
+	virtual FText GetBaseToolkitName() const override;
+	virtual FString GetWorldCentricTabPrefix() const override;
+	virtual FLinearColor GetWorldCentricTabColorScale() const override;
+	virtual void RegisterTabSpawners(const TSharedRef<FTabManager>& TabManager) override;
+	virtual void UnregisterTabSpawners(const TSharedRef<FTabManager>& TabManager) override;
+
+	/** FNotifyHook interface */
+	virtual void NotifyPostChange(const FPropertyChangedEvent& PropertyChangedEvent, UProperty* PropertyThatChanged) override;
+
+protected:
+	virtual void PostUndo(bool bSuccess) override;
+	virtual void PostRedo(bool bSuccess) override;
+
+private:
+	void SetCurveAtOrderIndex(int32 InCurveIndex, FRichCurve& InRichCurve, EModSettingsOutputEditorCurveSource InSource, UCurveFloat* InSharedCurve = nullptr);
+
+	/**	Spawns the tab allowing for editing/viewing the output curve(s) */
+	TSharedRef<SDockTab> SpawnTab_OutputCurve(const FSpawnTabArgs& Args);
+
+	/**	Spawns the tab allowing for editing/viewing the output curve(s) */
+	TSharedRef<SDockTab> SpawnTab_Properties(const FSpawnTabArgs& Args);
+
+	int32 GetControlModulationCount() const;
+
+	int32 GetCurveOrderIndex(EModSettingsEditorCurveOutput InCurveOutput, int32 InControlIndex) const;
+
+	int32 GetControlIndex(int32 InCurveIndex) const;
+
+	FName GetControlName(int32 InCurveIndex) const;
+
+	int32 GetCurveCount() const;
+
+	/** Get the orientation for the snap value controls. */
+	EOrientation GetSnapLabelOrientation() const;
+
+	/** Updates provided curve output's curve. Returns index of CurveModelIDs that was updated */
+	void UpdateCurve(int32 InCurveIndex);
+
+	/** Updates all curves */
+	void UpdateCurves();
+
+	/** Trims keys out-of-bounds in provided output */
+	void TrimKeys(FSoundModulationOutputBase& InOutput) const;
+
+	FSoundModulationOutputBase* FindModulationOutput(int32 InCurveIndex);
+
+	void GenerateExpressionCurve(const FSoundModulationOutputBase& InOutput, int32 InCurveIndex, bool bIsUnset = false);
+
+	TSharedPtr<FUICommandList> ToolbarCurveTargetCommands;
+
+	TSharedPtr<FCurveEditor> CurveEditor;
+	TSharedPtr<SCurveEditorPanel> CurvePanel;
+
+	TArray<TSharedPtr<FRichCurve>> ExpressionCurves;
+
+	TArray<FCurveModelID> CurveModels;
+
+	/** Properties tab */
+	TSharedPtr<IDetailsView> PropertiesView;
+
+	/** Settings Editor App Identifier */
+	static const FName AppIdentifier;
+	static const FName CurveTabId;
+	static const FName PropertiesTabId;
+};
diff --git a/Engine/Plugins/Runtime/AudioModulation/Source/AudioModulationEditor/Private/Factories/SoundControlBusFactory.cpp b/Engine/Plugins/Runtime/AudioModulation/Source/AudioModulationEditor/Private/Factories/SoundControlBusFactory.cpp
--- a/Engine/Plugins/Runtime/AudioModulation/Source/AudioModulationEditor/Private/Factories/SoundControlBusFactory.cpp
+++ b/Engine/Plugins/Runtime/AudioModulation/Source/AudioModulationEditor/Private/Factories/SoundControlBusFactory.cpp
@@ -58,4 +58,18 @@ USoundHPFControlBusFactory::USoundHPFControlBusFactory(const FObjectInitializer&
 UObject* USoundHPFControlBusFactory::FactoryCreateNew(UClass* Class, UObject* InParent, FName Name, EObjectFlags Flags, UObject* Context, FFeedbackContext* Warn)
 {
 	return NewObject<USoundHPFControlBus>(InParent, Name, Flags);
+}
+
+USoundControlBusFactory::USoundControlBusFactory(const FObjectInitializer& ObjectInitializer)
+	: Super(ObjectInitializer)
+{
+	SupportedClass = USoundControlBus::StaticClass();
+	bCreateNew = true;
+	bEditorImport = false;
+	bEditAfterNew = true;
+}
+
+UObject* USoundControlBusFactory::FactoryCreateNew(UClass* Class, UObject* InParent, FName Name, EObjectFlags Flags, UObject* Context, FFeedbackContext* Warn)
+{
+	return NewObject<USoundControlBus>(InParent, Name, Flags);
 }
\ No newline at end of file
diff --git a/Engine/Plugins/Runtime/AudioModulation/Source/AudioModulationEditor/Private/Factories/SoundControlBusFactory.h b/Engine/Plugins/Runtime/AudioModulation/Source/AudioModulationEditor/Private/Factories/SoundControlBusFactory.h
--- a/Engine/Plugins/Runtime/AudioModulation/Source/AudioModulationEditor/Private/Factories/SoundControlBusFactory.h
+++ b/Engine/Plugins/Runtime/AudioModulation/Source/AudioModulationEditor/Private/Factories/SoundControlBusFactory.h
@@ -40,5 +40,13 @@ class USoundLPFControlBusFactory : public UFactory
 {
 	GENERATED_UCLASS_BODY()
 
+	virtual UObject* FactoryCreateNew(UClass* Class, UObject* InParent, FName Name, EObjectFlags Flags, UObject* Context, FFeedbackContext* Warn) override;
+};
+
+UCLASS(hidecategories = Object, MinimalAPI)
+class USoundControlBusFactory : public UFactory
+{
+	GENERATED_UCLASS_BODY()
+
 	virtual UObject* FactoryCreateNew(UClass* Class, UObject* InParent, FName Name, EObjectFlags Flags, UObject* Context, FFeedbackContext* Warn) override;
 };
\ No newline at end of file
diff --git a/Engine/Plugins/Runtime/AudioModulation/Source/AudioModulationEditor/Private/Layouts/SoundControlModulationPatchLayout.cpp b/Engine/Plugins/Runtime/AudioModulation/Source/AudioModulationEditor/Private/Layouts/SoundControlModulationPatchLayout.cpp
new file mode 100644
--- /dev/null
+++ b/Engine/Plugins/Runtime/AudioModulation/Source/AudioModulationEditor/Private/Layouts/SoundControlModulationPatchLayout.cpp
@@ -0,0 +1,127 @@
+// Copyright 1998-2019 Epic Games, Inc. All Rights Reserved.
+#include "SoundControlModulationPatchLayout.h"
+
+#include "AudioModulationSettings.h"
+#include "DetailCategoryBuilder.h"
+#include "DetailLayoutBuilder.h"
+#include "DetailWidgetRow.h"
+#include "EditorStyleSet.h"
+#include "IAudioExtensionPlugin.h"
+#include "IDetailChildrenBuilder.h"
+#include "IDetailPropertyRow.h"
+#include "SoundModulationPatch.h"
+#include "Widgets/Input/SButton.h"
+#include "Widgets/Text/STextBlock.h"
+#include "Widgets/SBoxPanel.h"
+
+
+#define LOCTEXT_NAMESPACE "SoundControlModulationPatch"
+namespace
+{
+	void GetPropertyHandleMap(TSharedRef<IPropertyHandle> StructPropertyHandle, TMap<FName, TSharedPtr<IPropertyHandle>>& PropertyHandles)
+	{
+		uint32 NumChildren;
+		StructPropertyHandle->GetNumChildren(NumChildren);
+
+		for (uint32 ChildIndex = 0; ChildIndex < NumChildren; ++ChildIndex)
+		{
+			TSharedRef<IPropertyHandle> ChildHandle = StructPropertyHandle->GetChildHandle(ChildIndex).ToSharedRef();
+			const FName PropertyName = ChildHandle->GetProperty()->GetFName();
+			PropertyHandles.Add(PropertyName, ChildHandle);
+		}
+	}
+} // namespace <>
+
+void FSoundModulationPatchLayoutCustomization::CustomizeHeader(TSharedRef<IPropertyHandle> StructPropertyHandle, FDetailWidgetRow& HeaderRow, IPropertyTypeCustomizationUtils& CustomizationUtils)
+{
+	// Property handle here is the base struct. We are going to hide it since we're showing it's properties directly.
+// 	PropertyHandle->MarkHiddenByCustomization();
+}
+
+void FSoundModulationPatchLayoutCustomization::CustomizeChildren(TSharedRef<IPropertyHandle> StructPropertyHandle, IDetailChildrenBuilder& ChildBuilder, IPropertyTypeCustomizationUtils& StructCustomizationUtils)
+{
+	TMap<FName, TSharedPtr<IPropertyHandle>> PropertyHandles;
+	GetPropertyHandleMap(StructPropertyHandle, PropertyHandles);
+	CustomizeControl(PropertyHandles, ChildBuilder);
+}
+
+void FSoundModulationPatchLayoutCustomization::CustomizeControl(TMap<FName, TSharedPtr<IPropertyHandle>> &PropertyHandles, IDetailChildrenBuilder &ChildBuilder)
+{
+	TSharedPtr<IPropertyHandle>ValueHandle = PropertyHandles.FindChecked(GET_MEMBER_NAME_CHECKED(FSoundModulationPatchBase, DefaultInputValue));
+	ChildBuilder.AddProperty(ValueHandle.ToSharedRef());
+}
+
+void FSoundVolumeModulationPatchLayoutCustomization::CustomizeControl(TMap<FName, TSharedPtr<IPropertyHandle>> &PropertyHandles, IDetailChildrenBuilder &ChildBuilder)
+{
+	FSoundModulationPatchLayoutCustomization::CustomizeControl(PropertyHandles, ChildBuilder);
+
+	TSharedPtr<IPropertyHandle>InputsHandle = PropertyHandles.FindChecked(GET_MEMBER_NAME_CHECKED(FSoundVolumeModulationPatch, Inputs));
+	ChildBuilder.AddProperty(InputsHandle.ToSharedRef());
+
+	TSharedPtr<IPropertyHandle>OutputHandle = PropertyHandles.FindChecked(GET_MEMBER_NAME_CHECKED(FSoundVolumeModulationPatch, Output));
+	ChildBuilder.AddProperty(OutputHandle.ToSharedRef());
+}
+
+void FSoundPitchModulationPatchLayoutCustomization::CustomizeControl(TMap<FName, TSharedPtr<IPropertyHandle>> &PropertyHandles, IDetailChildrenBuilder &ChildBuilder)
+{
+	FSoundModulationPatchLayoutCustomization::CustomizeControl(PropertyHandles, ChildBuilder);
+
+	TSharedPtr<IPropertyHandle>InputsHandle = PropertyHandles.FindChecked(GET_MEMBER_NAME_CHECKED(FSoundPitchModulationPatch, Inputs));
+	ChildBuilder.AddProperty(InputsHandle.ToSharedRef());
+
+	TSharedPtr<IPropertyHandle>OutputHandle = PropertyHandles.FindChecked(GET_MEMBER_NAME_CHECKED(FSoundPitchModulationPatch, Output));
+	ChildBuilder.AddProperty(OutputHandle.ToSharedRef());
+}
+
+void FSoundHPFModulationPatchLayoutCustomization::CustomizeControl(TMap<FName, TSharedPtr<IPropertyHandle>> &PropertyHandles, IDetailChildrenBuilder &ChildBuilder)
+{
+	FSoundModulationPatchLayoutCustomization::CustomizeControl(PropertyHandles, ChildBuilder);
+
+	TSharedPtr<IPropertyHandle>InputsHandle = PropertyHandles.FindChecked(GET_MEMBER_NAME_CHECKED(FSoundHPFModulationPatch, Inputs));
+	ChildBuilder.AddProperty(InputsHandle.ToSharedRef());
+
+	TSharedPtr<IPropertyHandle>OutputHandle = PropertyHandles.FindChecked(GET_MEMBER_NAME_CHECKED(FSoundHPFModulationPatch, Output));
+	ChildBuilder.AddProperty(OutputHandle.ToSharedRef());
+}
+
+void FSoundLPFModulationPatchLayoutCustomization::CustomizeControl(TMap<FName, TSharedPtr<IPropertyHandle>> &PropertyHandles, IDetailChildrenBuilder &ChildBuilder)
+{
+	FSoundModulationPatchLayoutCustomization::CustomizeControl(PropertyHandles, ChildBuilder);
+
+	TSharedPtr<IPropertyHandle>InputsHandle = PropertyHandles.FindChecked(GET_MEMBER_NAME_CHECKED(FSoundLPFModulationPatch, Inputs));
+	ChildBuilder.AddProperty(InputsHandle.ToSharedRef());
+
+	TSharedPtr<IPropertyHandle>OutputHandle = PropertyHandles.FindChecked(GET_MEMBER_NAME_CHECKED(FSoundLPFModulationPatch, Output));
+	ChildBuilder.AddProperty(OutputHandle.ToSharedRef());
+}
+
+void FSoundControlModulationPatchLayoutCustomization::CustomizeControl(TMap<FName, TSharedPtr<IPropertyHandle>> &PropertyHandles, IDetailChildrenBuilder &ChildBuilder)
+{
+// Properties hidden as Generic Control Modulation is still in development
+// 	FName SetControlName;
+// 
+// 	TSharedPtr<IPropertyHandle> ControlProperty = PropertyHandles.FindChecked(GET_MEMBER_NAME_CHECKED(FSoundControlModulationPatch, Control));
+// 	ControlProperty->GetValue(SetControlName);
+// 
+// 	TSharedPtr<IPropertyHandle>ValueHandle = PropertyHandles.FindChecked(GET_MEMBER_NAME_CHECKED(FSoundControlModulationPatch, DefaultInputValue));
+// 
+// 	ChildBuilder.AddCustomRow(FText::FromName(SetControlName))
+// 	.NameContent()
+// 	[
+// 		SAssignNew(ControlComboBox, SSoundModulationControlComboBox)
+// 		.InitControlName(SetControlName)
+// 		.ControlNameProperty(ControlProperty)
+// 	]
+// 	.ValueContent()
+// 	.MinDesiredWidth(120)
+// 	[
+// 		ValueHandle->CreatePropertyValueWidget()
+// 	];
+// 
+// 	TSharedPtr<IPropertyHandle>InputsHandle = PropertyHandles.FindChecked(GET_MEMBER_NAME_CHECKED(FSoundControlModulationPatch, Inputs));
+// 	ChildBuilder.AddProperty(InputsHandle.ToSharedRef());
+// 
+// 	TSharedPtr<IPropertyHandle>OutputHandle = PropertyHandles.FindChecked(GET_MEMBER_NAME_CHECKED(FSoundControlModulationPatch, Output));
+// 	ChildBuilder.AddProperty(OutputHandle.ToSharedRef());
+}
+#undef LOCTEXT_NAMESPACE
diff --git a/Engine/Plugins/Runtime/AudioModulation/Source/AudioModulationEditor/Private/Layouts/SoundControlModulationPatchLayout.h b/Engine/Plugins/Runtime/AudioModulation/Source/AudioModulationEditor/Private/Layouts/SoundControlModulationPatchLayout.h
new file mode 100644
--- /dev/null
+++ b/Engine/Plugins/Runtime/AudioModulation/Source/AudioModulationEditor/Private/Layouts/SoundControlModulationPatchLayout.h
@@ -0,0 +1,89 @@
+// Copyright 1998-2019 Epic Games, Inc. All Rights Reserved.
+#pragma once
+
+#include "CoreMinimal.h"
+
+#include "IPropertyTypeCustomization.h"
+#include "Layout/Visibility.h"
+#include "Misc/Attribute.h"
+#include "PropertyHandle.h"
+#include "SoundModulationControlComboBox.h"
+
+// Forward Declarations
+class SSearchableComboBox;
+struct FSoundModulationSettings;
+
+
+class FSoundModulationPatchLayoutCustomization : public IPropertyTypeCustomization
+{
+public:
+	static TSharedRef<IPropertyTypeCustomization> MakeInstance()
+	{
+		return MakeShared<FSoundModulationPatchLayoutCustomization>();
+	}
+
+	//~ Begin IPropertyTypeCustomization
+	virtual void CustomizeHeader(TSharedRef<IPropertyHandle> StructPropertyHandle, FDetailWidgetRow& HeaderRow, IPropertyTypeCustomizationUtils& StructCustomizationUtils) override;
+	virtual void CustomizeChildren(TSharedRef<IPropertyHandle> StructPropertyHandle, IDetailChildrenBuilder& ChildBuilder, IPropertyTypeCustomizationUtils& StructCustomizationUtils) override;
+	//~ End IPropertyTypeCustomization
+
+	virtual void CustomizeControl(TMap<FName, TSharedPtr<IPropertyHandle>> &PropertyHandles, IDetailChildrenBuilder &ChildBuilder);
+};
+
+class FSoundVolumeModulationPatchLayoutCustomization : public FSoundModulationPatchLayoutCustomization
+{
+public:
+	static TSharedRef<IPropertyTypeCustomization> MakeInstance()
+	{
+		return MakeShared<FSoundVolumeModulationPatchLayoutCustomization>();
+	}
+
+	virtual void CustomizeControl(TMap<FName, TSharedPtr<IPropertyHandle>> &PropertyHandles, IDetailChildrenBuilder &ChildBuilder) override;
+};
+
+class FSoundPitchModulationPatchLayoutCustomization : public FSoundModulationPatchLayoutCustomization
+{
+public:
+	static TSharedRef<IPropertyTypeCustomization> MakeInstance()
+	{
+		return MakeShared<FSoundPitchModulationPatchLayoutCustomization>();
+	}
+
+	virtual void CustomizeControl(TMap<FName, TSharedPtr<IPropertyHandle>> &PropertyHandles, IDetailChildrenBuilder &ChildBuilder) override;
+};
+
+class FSoundLPFModulationPatchLayoutCustomization : public FSoundModulationPatchLayoutCustomization
+{
+public:
+	static TSharedRef<IPropertyTypeCustomization> MakeInstance()
+	{
+		return MakeShared<FSoundLPFModulationPatchLayoutCustomization>();
+	}
+
+	virtual void CustomizeControl(TMap<FName, TSharedPtr<IPropertyHandle>> &PropertyHandles, IDetailChildrenBuilder &ChildBuilder) override;
+};
+
+class FSoundHPFModulationPatchLayoutCustomization : public FSoundModulationPatchLayoutCustomization
+{
+public:
+	static TSharedRef<IPropertyTypeCustomization> MakeInstance()
+	{
+		return MakeShared<FSoundHPFModulationPatchLayoutCustomization>();
+	}
+
+	virtual void CustomizeControl(TMap<FName, TSharedPtr<IPropertyHandle>> &PropertyHandles, IDetailChildrenBuilder &ChildBuilder) override;
+};
+
+class FSoundControlModulationPatchLayoutCustomization : public FSoundModulationPatchLayoutCustomization
+{
+public:
+	static TSharedRef<IPropertyTypeCustomization> MakeInstance()
+	{
+		return MakeShared<FSoundControlModulationPatchLayoutCustomization>();
+	}
+
+	virtual void CustomizeControl(TMap<FName, TSharedPtr<IPropertyHandle>> &PropertyHandles, IDetailChildrenBuilder &ChildBuilder) override;
+
+protected:
+	TSharedPtr<SSearchableComboBox> ControlComboBox;
+};
diff --git a/Engine/Plugins/Runtime/AudioModulation/Source/AudioModulationEditor/Private/Layouts/SoundModulationControlComboBox.cpp b/Engine/Plugins/Runtime/AudioModulation/Source/AudioModulationEditor/Private/Layouts/SoundModulationControlComboBox.cpp
new file mode 100644
--- /dev/null
+++ b/Engine/Plugins/Runtime/AudioModulation/Source/AudioModulationEditor/Private/Layouts/SoundModulationControlComboBox.cpp
@@ -0,0 +1,106 @@
+// Copyright 1998-2019 Epic Games, Inc. All Rights Reserved.
+#include "SoundModulationControlComboBox.h"
+
+#include "AudioModulationSettings.h"
+#include "DetailLayoutBuilder.h"
+#include "DetailWidgetRow.h"
+
+
+#define LOCTEXT_NAMESPACE "SoundModulationParameter"
+void SSoundModulationControlComboBox::Construct(const FArguments& InArgs)
+{
+	ControlNameProperty = InArgs._ControlNameProperty;
+
+	ControlNameStrings.Add(MakeShared<FString>(TEXT("None")));
+
+	TSharedPtr<FString> SetControlValue;
+	if (const UAudioModulationSettings* ModSettings = GetDefault<UAudioModulationSettings>())
+	{
+		for (FName Name : ModSettings->ControlNames)
+		{
+			TSharedPtr<FString> NewName = MakeShared<FString>(Name.ToString());
+			if (Name == InArgs._InitControlName)
+			{
+				SetControlValue = NewName;
+			}
+			ControlNameStrings.Add(NewName);
+		}
+	}
+	ControlNameStrings.Sort([](const TSharedPtr<FString>& A, const TSharedPtr<FString>& B) { return A->Compare(*B.Get()) < 0; });
+
+	SSearchableComboBox::Construct(SSearchableComboBox::FArguments()
+		.OptionsSource(&ControlNameStrings)
+		.OnGenerateWidget(this, &SSoundModulationControlComboBox::MakeControlComboWidget)
+		.OnSelectionChanged(this, &SSoundModulationControlComboBox::OnControlChanged)
+		.OnComboBoxOpening(this, &SSoundModulationControlComboBox::OnControlComboOpening)
+		.ContentPadding(0)
+		.Content()
+		[
+			SNew(STextBlock)
+			.Text(this, &SSoundModulationControlComboBox::GetControlComboBoxContent)
+			.Font(IDetailLayoutBuilder::GetDetailFont())
+		]
+		.ToolTipText(ControlNameProperty->GetToolTipText())
+	);
+}
+
+FText SSoundModulationControlComboBox::GetControlComboBoxContent() const
+{
+	FName ControlSourceName;
+	if (ControlNameProperty->GetValue(ControlSourceName) == FPropertyAccess::Result::MultipleValues)
+	{
+		return LOCTEXT("MultipleValues", "Multiple Values");
+	}
+
+	return FText::FromString(*GetControlString(ControlSourceName).Get());
+}
+
+TSharedPtr<FString> SSoundModulationControlComboBox::GetControlString(FName ControlName) const
+{
+	FString ControlString = ControlName.ToString();
+
+	// go through profile and see if it has mine
+	for (int32 Index = 1; Index < ControlNameStrings.Num(); ++Index)
+	{
+		if (ControlString == *ControlNameStrings[Index])
+		{
+			return ControlNameStrings[Index];
+		}
+	}
+
+	return ControlNameStrings[0];
+}
+
+TSharedRef<SWidget> SSoundModulationControlComboBox::MakeControlComboWidget(TSharedPtr<FString> InItem)
+{
+	return SNew(STextBlock).Text(FText::FromString(*InItem)).Font(IDetailLayoutBuilder::GetDetailFont());
+}
+
+void SSoundModulationControlComboBox::OnControlChanged(TSharedPtr<FString> NewSelection, ESelectInfo::Type SelectInfo)
+{
+	// if set from code, ignore validation
+	if (SelectInfo != ESelectInfo::Direct)
+	{
+		FString NewValue = *NewSelection.Get();
+		if (NewValue == TEXT("None"))
+		{
+			NewValue = TEXT("");
+		}
+
+		ensure(ControlNameProperty->SetValue(NewValue) == FPropertyAccess::Result::Success);
+	}
+}
+
+void SSoundModulationControlComboBox::OnControlComboOpening()
+{
+	FName ControlName;
+	if (ControlNameProperty->GetValue(ControlName) != FPropertyAccess::Result::MultipleValues)
+	{
+		TSharedPtr<FString> ComboStringPtr = GetControlString(ControlName);
+		if (ComboStringPtr.IsValid())
+		{
+			SetSelectedItem(ComboStringPtr);
+		}
+	}
+}
+#undef LOCTEXT_NAMESPACE
\ No newline at end of file
diff --git a/Engine/Plugins/Runtime/AudioModulation/Source/AudioModulationEditor/Private/Layouts/SoundModulationControlComboBox.h b/Engine/Plugins/Runtime/AudioModulation/Source/AudioModulationEditor/Private/Layouts/SoundModulationControlComboBox.h
new file mode 100644
--- /dev/null
+++ b/Engine/Plugins/Runtime/AudioModulation/Source/AudioModulationEditor/Private/Layouts/SoundModulationControlComboBox.h
@@ -0,0 +1,38 @@
+// Copyright 1998-2019 Epic Games, Inc. All Rights Reserved.
+#pragma once
+
+#include "CoreMinimal.h"
+
+#include "PropertyHandle.h"
+#include "SSearchableComboBox.h"
+#include "Types/SlateEnums.h"
+#include "Widgets/SWidget.h"
+
+// Forward Declarations
+FDetailWidgetRow;
+
+
+class SSoundModulationControlComboBox : public SSearchableComboBox
+{
+public:
+	SLATE_BEGIN_ARGS(SSoundModulationControlComboBox)
+		: _InitControlName(NAME_None)
+	{}
+		SLATE_ARGUMENT(TSharedPtr<IPropertyHandle>, ControlNameProperty)
+		SLATE_ARGUMENT(FName, InitControlName)
+	SLATE_END_ARGS()
+
+	void Construct(const FArguments& InArgs);
+
+protected:
+	FText GetControlComboBoxContent() const;
+	TSharedPtr<FString> GetControlString(FName ControlName) const;
+
+	TSharedRef<SWidget> MakeControlComboWidget(TSharedPtr<FString> InItem);
+
+	void OnControlChanged(TSharedPtr<FString> NewSelection, ESelectInfo::Type SelectInfo);
+	void OnControlComboOpening();
+
+	TSharedPtr<IPropertyHandle> ControlNameProperty;
+	TArray<TSharedPtr<FString>> ControlNameStrings;
+};
diff --git a/Engine/Plugins/Runtime/AudioModulation/Source/AudioModulationEditor/Private/Layouts/SoundModulationParameterLayout.cpp b/Engine/Plugins/Runtime/AudioModulation/Source/AudioModulationEditor/Private/Layouts/SoundModulationParameterLayout.cpp
new file mode 100644
--- /dev/null
+++ b/Engine/Plugins/Runtime/AudioModulation/Source/AudioModulationEditor/Private/Layouts/SoundModulationParameterLayout.cpp
@@ -0,0 +1,100 @@
+// Copyright 1998-2019 Epic Games, Inc. All Rights Reserved.
+#include "SoundModulationParameterLayout.h"
+
+#include "AudioModulationSettings.h"
+#include "DetailCategoryBuilder.h"
+#include "DetailLayoutBuilder.h"
+#include "DetailWidgetRow.h"
+#include "EditorStyleSet.h"
+#include "IAudioExtensionPlugin.h"
+#include "IDetailChildrenBuilder.h"
+#include "IDetailPropertyRow.h"
+#include "Widgets/Input/SButton.h"
+#include "Widgets/Text/STextBlock.h"
+#include "Widgets/SBoxPanel.h"
+
+
+#define LOCTEXT_NAMESPACE "SoundModulationParameter"
+void FSoundModulationParameterLayoutCustomization::CustomizeHeader(TSharedRef<IPropertyHandle> StructPropertyHandle, FDetailWidgetRow& HeaderRow, IPropertyTypeCustomizationUtils& CustomizationUtils)
+{
+	// Property handle here is the base struct. We are going to hide it since we're showing it's properties directly.
+// 	PropertyHandle->MarkHiddenByCustomization();
+}
+
+void FSoundModulationParameterLayoutCustomization::CustomizeChildren(TSharedRef<IPropertyHandle> StructPropertyHandle, IDetailChildrenBuilder& ChildBuilder, IPropertyTypeCustomizationUtils& StructCustomizationUtils)
+{
+	uint32 NumChildren;
+	StructPropertyHandle->GetNumChildren(NumChildren);
+
+	TMap<FName, TSharedPtr<IPropertyHandle>> PropertyHandles;
+
+	for (uint32 ChildIndex = 0; ChildIndex < NumChildren; ++ChildIndex)
+	{
+		TSharedRef<IPropertyHandle> ChildHandle = StructPropertyHandle->GetChildHandle(ChildIndex).ToSharedRef();
+		const FName PropertyName = ChildHandle->GetProperty()->GetFName();
+		PropertyHandles.Add(PropertyName, ChildHandle);
+	}
+
+	FName SetControlName;
+
+	TSharedPtr<IPropertyHandle> ControlProperty = PropertyHandles.FindChecked(GET_MEMBER_NAME_CHECKED(FSoundModulationParameter, Control));
+	ControlProperty->GetValue(SetControlName);
+
+	TSharedPtr<IPropertyHandle>ValueHandle = PropertyHandles.FindChecked(GET_MEMBER_NAME_CHECKED(FSoundModulationParameter, Value));
+	ValueHandle->SetOnPropertyValueChanged(FSimpleDelegate::CreateRaw(this, &FSoundModulationParameterLayoutCustomization::OnValueChanged, StructPropertyHandle));
+	ValueHandle->SetToolTipText(StructPropertyHandle->GetToolTipText());
+
+	ChildBuilder.AddCustomRow(StructPropertyHandle->GetPropertyDisplayName())
+	.NameContent()
+	[
+		SNew(STextBlock)
+		.Text(StructPropertyHandle->GetPropertyDisplayName())
+		.Font(IDetailLayoutBuilder::GetDetailFont())
+	]
+	.ValueContent()
+	.MinDesiredWidth(280)
+	[
+		SNew(SHorizontalBox)
+		+ SHorizontalBox::Slot()
+		.VAlign(VAlign_Center)
+		.FillWidth(0.35f)
+		.Padding(2.0f, 0.0f)
+		[
+			ValueHandle->CreatePropertyValueWidget()
+		]
+		+ SHorizontalBox::Slot()
+		.FillWidth(0.65f)
+		[
+			SAssignNew(ControlComboBox, SSoundModulationControlComboBox)
+			.InitControlName(SetControlName)
+			.ControlNameProperty(ControlProperty)
+		]
+	];
+}
+
+void FSoundModulationParameterLayoutCustomization::OnValueChanged(TSharedRef<IPropertyHandle> InStructPropertyHandle)
+{
+	// Potentially accessing the value while garbage collecting or saving the package could trigger a crash.
+	// so we fail to get the value when that is occurring.
+	if (GIsSavingPackage || IsGarbageCollecting())
+	{
+		return;
+	}
+
+	// Way to enforce value was set by editor and remains in-bounds but allows each callsite to set and check
+	// at runtime (as desired being that input bus values can be out-of-bounds by design)
+	if (InStructPropertyHandle->IsValidHandle())
+	{
+		TArray<void*> RawData;
+		InStructPropertyHandle->AccessRawData(RawData);
+		for (void* RawPtr : RawData)
+		{
+			if (RawPtr)
+			{
+				FSoundModulationParameter& ModParam = *reinterpret_cast<FSoundModulationParameter*>(RawPtr);
+				ModParam.SetBounds(ModParam.GetMinValue(), ModParam.GetMaxValue());
+			}
+		}
+	}
+}
+#undef LOCTEXT_NAMESPACE
diff --git a/Engine/Plugins/Runtime/AudioModulation/Source/AudioModulationEditor/Private/Layouts/SoundModulationParameterLayout.h b/Engine/Plugins/Runtime/AudioModulation/Source/AudioModulationEditor/Private/Layouts/SoundModulationParameterLayout.h
new file mode 100644
--- /dev/null
+++ b/Engine/Plugins/Runtime/AudioModulation/Source/AudioModulationEditor/Private/Layouts/SoundModulationParameterLayout.h
@@ -0,0 +1,34 @@
+// Copyright 1998-2019 Epic Games, Inc. All Rights Reserved.
+#pragma once
+
+#include "CoreMinimal.h"
+
+#include "IPropertyTypeCustomization.h"
+#include "Layout/Visibility.h"
+#include "Misc/Attribute.h"
+#include "PropertyHandle.h"
+#include "SoundModulationControlComboBox.h"
+
+// Forward Declarations
+class SSearchableComboBox;
+struct FSoundModulationParameter;
+
+
+class FSoundModulationParameterLayoutCustomization : public IPropertyTypeCustomization
+{
+public:
+	static TSharedRef<IPropertyTypeCustomization> MakeInstance()
+	{
+		return MakeShared<FSoundModulationParameterLayoutCustomization>();
+	}
+
+	//~ Begin IPropertyTypeCustomization
+	virtual void CustomizeHeader(TSharedRef<IPropertyHandle> StructPropertyHandle, FDetailWidgetRow& HeaderRow, IPropertyTypeCustomizationUtils& StructCustomizationUtils) override;
+	virtual void CustomizeChildren(TSharedRef<IPropertyHandle> StructPropertyHandle, IDetailChildrenBuilder& ChildBuilder, IPropertyTypeCustomizationUtils& StructCustomizationUtils) override;
+	//~ End IPropertyTypeCustomization
+
+protected:
+	void OnValueChanged(TSharedRef<IPropertyHandle> InStructPropertyHandle);
+
+	TSharedPtr<SSearchableComboBox> ControlComboBox;
+};
diff --git a/Engine/Plugins/Runtime/AudioModulation/Source/AudioModulationEditor/Private/Layouts/SoundModulationTransformLayout.cpp b/Engine/Plugins/Runtime/AudioModulation/Source/AudioModulationEditor/Private/Layouts/SoundModulationTransformLayout.cpp
--- a/Engine/Plugins/Runtime/AudioModulation/Source/AudioModulationEditor/Private/Layouts/SoundModulationTransformLayout.cpp
+++ b/Engine/Plugins/Runtime/AudioModulation/Source/AudioModulationEditor/Private/Layouts/SoundModulationTransformLayout.cpp
@@ -2,41 +2,32 @@
 #include "SoundModulationTransformLayout.h"
 
 #include "Audio.h"
+#include "AudioModulationStyle.h"
 #include "DetailCategoryBuilder.h"
 #include "DetailLayoutBuilder.h"
 #include "DetailWidgetRow.h"
+#include "EditorStyleSet.h"
 #include "IDetailChildrenBuilder.h"
 #include "IDetailPropertyRow.h"
+#include "InputCoreTypes.h"
 #include "PropertyRestriction.h"
+#include "SCurveEditor.h"
+#include "ScopedTransaction.h"
+#include "SoundModulationPatch.h"
 #include "SoundModulationTransform.h"
+#include "Widgets/Input/SButton.h"
+#include "Widgets/Input/SEditableTextBox.h"
+#include "Widgets/Input/SNumericEntryBox.h"
+#include "Widgets/Layout/SSeparator.h"
+#include "Widgets/Layout/SWidgetSwitcher.h"
+#include "Widgets/Notifications/SNotificationList.h"
+#include "Widgets/Text/SRichTextBlock.h"
 #include "Widgets/Text/STextBlock.h"
+#include "Widgets/Views/SListView.h"
 
 
-FSoundModulationOutputTransformLayoutCustomization::FSoundModulationOutputTransformLayoutCustomization()
-{
-}
-
-bool FSoundModulationOutputTransformLayoutCustomization::IsCurveSet(TSharedPtr<IPropertyHandle> CurveHandle, const TArray<FString>& Filters) const
-{
-	if (!CurveHandle.IsValid())
-	{
-		return false;
-	}
-
-	FString CurveString;
-	CurveHandle->GetValueAsFormattedString(CurveString);
-	for (const FString& Filter : Filters)
-	{
-		if (CurveString.Contains(*Filter))
-		{
-			return true;
-		}
-	}
-
-	return false;
-}
-
-void FSoundModulationOutputTransformLayoutCustomization::CustomizeHeader(TSharedRef<IPropertyHandle> PropertyHandle, FDetailWidgetRow& HeaderRow, IPropertyTypeCustomizationUtils& CustomizationUtils)
+#define LOCTEXT_NAMESPACE "SoundModulationOutputTransform"
+void FSoundModulationOutputTransformLayoutCustomization::CustomizeHeader(TSharedRef<IPropertyHandle> StructPropertyHandle, FDetailWidgetRow& HeaderRow, IPropertyTypeCustomizationUtils& CustomizationUtils)
 {
 }
 
@@ -59,8 +50,21 @@ void FSoundModulationOutputTransformLayoutCustomization::CustomizeChildren(TShar
 
 	TSharedRef<IPropertyHandle>InputMaxHandle = PropertyHandles.FindChecked(GET_MEMBER_NAME_CHECKED(FSoundModulationOutputTransform, InputMax)).ToSharedRef();
 	ChildBuilder.AddProperty(InputMaxHandle);
-	TSharedRef<IPropertyHandle> ExpScalarHandle = PropertyHandles.FindChecked(GET_MEMBER_NAME_CHECKED(FSoundModulationOutputTransform, Scalar)).ToSharedRef();
-	ChildBuilder.AddProperty(ExpScalarHandle);
+
+	TSharedRef<IPropertyHandle>CurveHandle = PropertyHandles.FindChecked(GET_MEMBER_NAME_CHECKED(FSoundModulationOutputTransform, Curve)).ToSharedRef();
+	ChildBuilder.AddProperty(CurveHandle);
+
+	TSharedRef<IPropertyHandle> ScalarHandle = PropertyHandles.FindChecked(GET_MEMBER_NAME_CHECKED(FSoundModulationOutputTransform, Scalar)).ToSharedRef();
+	const TArray<FString> ScalarFilters = TArray<FString>({ TEXT("Exp"), TEXT("Log") });
+
+	ChildBuilder.AddProperty(ScalarHandle)
+		.EditCondition(TAttribute<bool>::Create([this, CurveHandle, ScalarFilters]() { return IsScaleableCurve(CurveHandle, ScalarFilters); }), nullptr)
+		.Visibility(TAttribute<EVisibility>::Create([this, CurveHandle, ScalarFilters]() { return IsScaleableCurve(CurveHandle, ScalarFilters) ? EVisibility::Visible : EVisibility::Hidden; }));
+
+	TSharedRef<IPropertyHandle> SharedCurveHandle = PropertyHandles.FindChecked(GET_MEMBER_NAME_CHECKED(FSoundModulationOutputTransform, CurveShared)).ToSharedRef();
+	ChildBuilder.AddProperty(SharedCurveHandle)
+		.EditCondition(TAttribute<bool>::Create([this, CurveHandle, ScalarFilters]() { return IsSharedCurve(CurveHandle); }), nullptr)
+		.Visibility(TAttribute<EVisibility>::Create([this, CurveHandle, ScalarFilters]() { return IsSharedCurve(CurveHandle) ? EVisibility::Visible : EVisibility::Hidden; }));
 
 	TSharedRef<IPropertyHandle>OutputMinHandle = PropertyHandles.FindChecked(GET_MEMBER_NAME_CHECKED(FSoundModulationOutputTransform, OutputMin)).ToSharedRef();
 	ChildBuilder.AddProperty(OutputMinHandle);
@@ -68,3 +72,36 @@ void FSoundModulationOutputTransformLayoutCustomization::CustomizeChildren(TShar
 	TSharedRef<IPropertyHandle>OutputMaxHandle = PropertyHandles.FindChecked(GET_MEMBER_NAME_CHECKED(FSoundModulationOutputTransform, OutputMax)).ToSharedRef();
 	ChildBuilder.AddProperty(OutputMaxHandle);
 }
+
+bool FSoundModulationOutputTransformLayoutCustomization::IsScaleableCurve(TSharedPtr<IPropertyHandle> CurveHandle, const TArray<FString>& Filters) const
+{
+	if (!CurveHandle.IsValid())
+	{
+		return false;
+	}
+
+	FString CurveString;
+	CurveHandle->GetValueAsFormattedString(CurveString);
+	for (const FString& Filter : Filters)
+	{
+		if (CurveString.Contains(*Filter))
+		{
+			return true;
+		}
+	}
+
+	return false;
+}
+
+bool FSoundModulationOutputTransformLayoutCustomization::IsSharedCurve(TSharedPtr<IPropertyHandle> CurveHandle) const
+{
+	if (!CurveHandle.IsValid())
+	{
+		return false;
+	}
+
+	FString CurveString;
+	CurveHandle->GetValueAsFormattedString(CurveString);
+	return CurveString == TEXT("Shared");
+}
+#undef LOCTEXT_NAMESPACE
diff --git a/Engine/Plugins/Runtime/AudioModulation/Source/AudioModulationEditor/Private/Layouts/SoundModulationTransformLayout.h b/Engine/Plugins/Runtime/AudioModulation/Source/AudioModulationEditor/Private/Layouts/SoundModulationTransformLayout.h
--- a/Engine/Plugins/Runtime/AudioModulation/Source/AudioModulationEditor/Private/Layouts/SoundModulationTransformLayout.h
+++ b/Engine/Plugins/Runtime/AudioModulation/Source/AudioModulationEditor/Private/Layouts/SoundModulationTransformLayout.h
@@ -2,20 +2,24 @@
 #pragma once
 
 #include "CoreMinimal.h"
+
+#include "Curves/SimpleCurve.h"
 #include "IPropertyTypeCustomization.h"
 #include "Layout/Visibility.h"
 #include "Misc/Attribute.h"
 #include "PropertyHandle.h"
+#include "SCurveEditor.h"
+
 
+// Forward Declarations
+struct FSoundModulationOutputTransform;
 
 class FSoundModulationOutputTransformLayoutCustomization : public IPropertyTypeCustomization
 {
 public:
-	FSoundModulationOutputTransformLayoutCustomization();
-
 	static TSharedRef<IPropertyTypeCustomization> MakeInstance()
 	{
-		return MakeShareable(new FSoundModulationOutputTransformLayoutCustomization);
+		return MakeShared<FSoundModulationOutputTransformLayoutCustomization>();
 	}
 
 	//~ Begin IPropertyTypeCustomization
@@ -24,5 +28,6 @@ public:
 	//~ End IPropertyTypeCustomization
 
 private:
-	bool IsCurveSet(TSharedPtr<IPropertyHandle> CurveHandle, const TArray<FString>& Filter) const;
-};
\ No newline at end of file
+	bool IsScaleableCurve(TSharedPtr<IPropertyHandle> CurveHandle, const TArray<FString>& Filter) const;
+	bool IsSharedCurve(TSharedPtr<IPropertyHandle> CurveHandle) const;
+};
