Fix missing class specifier in forward declare

diff --git a/Engine/Plugins/Runtime/AudioModulation/Source/AudioModulationEditor/Private/Layouts/SoundModulationControlComboBox.h b/Engine/Plugins/Runtime/AudioModulation/Source/AudioModulationEditor/Private/Layouts/SoundModulationControlComboBox.h
--- a/Engine/Plugins/Runtime/AudioModulation/Source/AudioModulationEditor/Private/Layouts/SoundModulationControlComboBox.h
+++ b/Engine/Plugins/Runtime/AudioModulation/Source/AudioModulationEditor/Private/Layouts/SoundModulationControlComboBox.h
@@ -8,8 +8,9 @@
 #include "Types/SlateEnums.h"
 #include "Widgets/SWidget.h"
 
+
 // Forward Declarations
-FDetailWidgetRow;
+class FDetailWidgetRow;
 
 
 class SSoundModulationControlComboBox : public SSearchableComboBox
