Fixed localization conflicts

#rnx

diff --git a/Engine/Plugins/Enterprise/DatasmithImporter/Source/DatasmithImporter/Private/DatasmithTextureImporter.cpp b/Engine/Plugins/Enterprise/DatasmithImporter/Source/DatasmithImporter/Private/DatasmithTextureImporter.cpp
--- a/Engine/Plugins/Enterprise/DatasmithImporter/Source/DatasmithImporter/Private/DatasmithTextureImporter.cpp
+++ b/Engine/Plugins/Enterprise/DatasmithImporter/Source/DatasmithImporter/Private/DatasmithTextureImporter.cpp
@@ -147,7 +147,7 @@ bool FDatasmithTextureImporter::GetTextureData(const TSharedPtr<IDatasmithTextur
 		// try opening from absolute path
 		if (!(FFileHelper::LoadFileToArray(TextureData, *ImageFileName) && TextureData.Num() > 0))
 		{
-			ImportContext.LogWarning(FText::Format(LOCTEXT("InvalidData", "Unable to find Texture file {0}."), FText::FromString(Filename)));
+			ImportContext.LogWarning(FText::Format(LOCTEXT("UnableToFindTexture", "Unable to find Texture file {0}."), FText::FromString(Filename)));
 			return false;
 		}
 
diff --git a/Engine/Plugins/Runtime/AudioModulation/Source/AudioModulationEditor/Private/Editors/ModulationSettingsCurveEditorViewStacked.cpp b/Engine/Plugins/Runtime/AudioModulation/Source/AudioModulationEditor/Private/Editors/ModulationSettingsCurveEditorViewStacked.cpp
--- a/Engine/Plugins/Runtime/AudioModulation/Source/AudioModulationEditor/Private/Editors/ModulationSettingsCurveEditorViewStacked.cpp
+++ b/Engine/Plugins/Runtime/AudioModulation/Source/AudioModulationEditor/Private/Editors/ModulationSettingsCurveEditorViewStacked.cpp
@@ -57,7 +57,7 @@ void FModCurveEditorModel::Refresh(EModSettingsEditorCurveOutput InCurveOutput,
 		case EModSettingsEditorCurveOutput::Highpass:
 		{
 			Color = FAudioModulationStyle::GetHPFBusColor();
-			ShortNameBase = LOCTEXT("ModulationOutputPitchBus", "Highpass");
+			ShortNameBase = LOCTEXT("ModulationOutputHPFBus", "Highpass");
 			IntentionName = TEXT("AudioHPFFreq");
 		}
 		break;
@@ -65,7 +65,7 @@ void FModCurveEditorModel::Refresh(EModSettingsEditorCurveOutput InCurveOutput,
 		case EModSettingsEditorCurveOutput::Lowpass:
 		{
 			Color = FAudioModulationStyle::GetLPFBusColor();
-			ShortNameBase = LOCTEXT("ModulationOutputPitchBus", "Lowpass");
+			ShortNameBase = LOCTEXT("ModulationOutputLPFBus", "Lowpass");
 			IntentionName = TEXT("AudioLPFFreq");
 		}
 		break;
@@ -106,7 +106,7 @@ void FModCurveEditorModel::Refresh(EModSettingsEditorCurveOutput InCurveOutput,
 		default:
 		{
 			bKeyDrawEnabled = 0;
-			ShortDisplayName = FText::Format(LOCTEXT("ModulationOutputCurveExpressionDisplayName", "{0} (Shared - Unset)"), ShortNameBase);
+			ShortDisplayName = FText::Format(LOCTEXT("ModulationOutputCurveUnsetDisplayName", "{0} (Shared - Unset)"), ShortNameBase);
 		}
 	}
 }
