- Fix bug with removing all modulateable sound references instead of single instances in mod proxies (can have the same sound reference a proxy more than once, and this is okay)
- Fix copy pasta in LFO deactivation
- Divert code paths more clearly for when previewing sounds and forcing activation of modulation types and fix bug with certain proxies sticking around or attempting to deactivate early
- Fix uninitialized BaseFreq member in IOscBase

diff --git a/Engine/Plugins/Runtime/AudioModulation/Source/AudioModulation/Private/AudioModulationImpl.cpp b/Engine/Plugins/Runtime/AudioModulation/Source/AudioModulation/Private/AudioModulationImpl.cpp
--- a/Engine/Plugins/Runtime/AudioModulation/Source/AudioModulation/Private/AudioModulationImpl.cpp
+++ b/Engine/Plugins/Runtime/AudioModulation/Source/AudioModulation/Private/AudioModulationImpl.cpp
@@ -99,6 +99,20 @@ namespace AudioModulation
 	}
 
 #if WITH_EDITOR
+	void FAudioModulationImpl::EndPreviewSession()
+	{
+		PreviewSound = nullptr;
+		PreviewSettings = FModulationSettingsProxy();
+
+		SoundSettings.Reset();
+
+		ActiveBuses.Reset();
+		ActiveBusMixes.Reset();
+		ActiveLFOs.Reset();
+
+		UE_LOG(LogAudioModulation, Verbose, TEXT("Releasing all modulation proxies: Sound preview session has ended."));
+	}
+
 	void FAudioModulationImpl::OnEditPluginSettings(const USoundModulationPluginSourceSettingsBase& InSettings)
 	{
 		// Find if sound is being referenced and auditioned and stop immediately.
@@ -161,8 +175,7 @@ namespace AudioModulation
 			if (bShouldStop)
 			{
 				PreviewSound->Stop();
-				PreviewSound = nullptr;
-				PreviewSettings = FModulationSettingsProxy();
+				EndPreviewSession();
 			}
 			else
 			{
@@ -196,9 +209,17 @@ namespace AudioModulation
 #if WITH_EDITOR
 		if (InSound.IsPreviewSound())
 		{
+			EndPreviewSession();
 			PreviewSound = &InSound;
 			PreviewSettings = FModulationSettingsProxy(*Settings);
 		}
+		// If currently previewing, activation/deactivation is by-passing
+		// auto-activation system and always activating/deactivating,
+		// so ignore requests from non-preview to avoid mixed state.
+		else if (PreviewSound)
+		{
+			return;
+		}
 #endif // WITH_EDITOR
 
 		const uint32 SoundId = InSound.GetObjectId();
@@ -209,10 +230,7 @@ namespace AudioModulation
 
 		for (const USoundControlBusMix* Mix : Settings->Mixes)
 		{
-			if (Mix && (Mix->bAutoActivate || InSound.IsPreviewSound()))
-			{
-				ActivateBusMix(*Mix, &InSound);
-			}
+			ActivateBusMix(*Mix, &InSound);
 		}
 
 		auto CheckRefActivate = [this, &InSound](const USoundControlBusBase* Bus)
@@ -222,20 +240,14 @@ namespace AudioModulation
 				return;
 			}
 
-			if (Bus->bAutoActivate || InSound.IsPreviewSound())
-			{
-				ActivateBus(*Bus, &InSound);
-			}
+			ActivateBus(*Bus, &InSound);
 
 			for (const USoundBusModulatorBase* Modulator : Bus->Modulators)
 			{
 				if (const USoundBusModulatorLFO* LFO = Cast<USoundBusModulatorLFO>(Modulator))
 				{
-					if (LFO->bAutoActivate || InSound.IsPreviewSound())
-					{
-						FModulatorLFOProxy LFOProxy(*LFO);
-						ActivateLFO(*LFO, &InSound);
-					}
+					FModulatorLFOProxy LFOProxy(*LFO);
+					ActivateLFO(*LFO, &InSound);
 				}
 			}
 		};
@@ -286,6 +298,20 @@ namespace AudioModulation
 		check(IsInAudioThread());
 		check(InSound.GetObjectId() != INDEX_NONE);
 
+#if WITH_EDITOR
+		// If previewing, don't attempt to deactivate anything, leaving
+		// all modulation proxies active until end of session, at which
+		// point all will be reset.
+		if (PreviewSound)
+		{
+			if (&InSound == PreviewSound)
+			{
+				EndPreviewSession();
+			}
+			return;
+		}
+#endif // WITH_EDITOR
+
 		// Settings can be null if sound settings were modified via the editor while auditioning or in PIE
 		if (const FModulationSettingsProxy* Settings = SoundSettings.Find(InSound.GetObjectId()))
 		{
@@ -321,15 +347,12 @@ namespace AudioModulation
 					DeactivateBus(Input.BusId, &InSound);
 				}
 			}
-		}
 
-#if WITH_EDITOR
-		if (InSound.IsPreviewSound() && &InSound == PreviewSound)
-		{
-			PreviewSound = nullptr;
-			PreviewSettings = FModulationSettingsProxy();
+			if (InSound.GetPlayCount() == 0)
+			{
+				SoundSettings.Remove(InSound.GetObjectId());
+			}
 		}
-#endif // WITH_EDITOR
 	}
 
 #if !UE_BUILD_SHIPPING
@@ -542,6 +565,13 @@ namespace AudioModulation
 
 			if (FControlBusProxy* Bus = ActiveBuses.Find(InBusId))
 			{
+				// Only pass along to referenced LFOs if deactivating
+				// via notification of sound release
+				for (FLFOId LFOId : Bus->GetLFOIds())
+				{
+					DeactivateLFO(LFOId, InSound);
+				}
+
 				if (!InSound)
 				{
 					if (!Bus->GetAutoActivate())
@@ -556,16 +586,6 @@ namespace AudioModulation
 						ActiveBuses.Remove(InBusId);
 					}
 				}
-
-				// Only pass along to referenced LFOs if deactivating
-				// via notification of sound release
-				if (InSound)
-				{
-					for (FLFOId LFOId : Bus->GetLFOIds())
-					{
-						DeactivateLFO(LFOId, InSound);
-					}
-				}
 			}
 		});
 	}
@@ -582,14 +602,14 @@ namespace AudioModulation
 				{
 					if (!LFO->GetAutoActivate())
 					{
-						ActiveBuses.Remove(InLFOId);
+						ActiveLFOs.Remove(InLFOId);
 					}
 				}
 				else if (LFO->OnReleaseSound(*InSound) == 0)
 				{
 					if (LFO->GetAutoActivate())
 					{
-						ActiveBuses.Remove(InLFOId);
+						ActiveLFOs.Remove(InLFOId);
 					}
 				}
 			}
diff --git a/Engine/Plugins/Runtime/AudioModulation/Source/AudioModulation/Private/AudioModulationInternal.h b/Engine/Plugins/Runtime/AudioModulation/Source/AudioModulation/Private/AudioModulationInternal.h
--- a/Engine/Plugins/Runtime/AudioModulation/Source/AudioModulation/Private/AudioModulationInternal.h
+++ b/Engine/Plugins/Runtime/AudioModulation/Source/AudioModulation/Private/AudioModulationInternal.h
@@ -69,6 +69,10 @@ namespace AudioModulation
 		/* Calculates modulation value, storing it in the provided float reference and returns if value changed */
 		bool CalculateModulationValue(FModulationPatchProxy& Proxy, float& OutValue) const;
 
+#if WITH_EDITOR
+		void EndPreviewSession();
+#endif // WITH_EDITOR
+
 		void RunCommandOnAudioThread(TFunction<void()> Cmd);
 
 		BusMixProxyMap ActiveBusMixes;
diff --git a/Engine/Plugins/Runtime/AudioModulation/Source/AudioModulation/Private/SoundModulationProxy.h b/Engine/Plugins/Runtime/AudioModulation/Source/AudioModulation/Private/SoundModulationProxy.h
--- a/Engine/Plugins/Runtime/AudioModulation/Source/AudioModulation/Private/SoundModulationProxy.h
+++ b/Engine/Plugins/Runtime/AudioModulation/Source/AudioModulation/Private/SoundModulationProxy.h
@@ -93,7 +93,7 @@ namespace AudioModulation
 		{
 			if (bAutoActivate)
 			{
-				const int32 NumRemoved = RefSounds.Remove(&Sound);
+				const int32 NumRemoved = RefSounds.RemoveSingleSwap(&Sound, false /* bAllowShrinking */);
 				check(NumRemoved == 1);
 			}
 
@@ -121,7 +121,7 @@ namespace AudioModulation
 
 			if (bAutoActivate)
 			{
-				RefSounds.AddUnique(&Sound);
+				RefSounds.Add(&Sound);
 			}
 
 			return RefSounds.Num();
diff --git a/Engine/Plugins/Runtime/AudioModulation/Source/AudioModulation/Private/SoundModulatorLFO.cpp b/Engine/Plugins/Runtime/AudioModulation/Source/AudioModulation/Private/SoundModulatorLFO.cpp
--- a/Engine/Plugins/Runtime/AudioModulation/Source/AudioModulation/Private/SoundModulatorLFO.cpp
+++ b/Engine/Plugins/Runtime/AudioModulation/Source/AudioModulation/Private/SoundModulatorLFO.cpp
@@ -50,14 +50,16 @@ namespace AudioModulation
 {
 	FModulatorLFOProxy::FModulatorLFOProxy()
 		: Offset(0.0f)
-		, Value(0.0f)
+		, Value(1.0f)
 	{
+		LFO.SetFrequency(1.0f);
+		LFO.Start();
 	}
 
 	FModulatorLFOProxy::FModulatorLFOProxy(const USoundBusModulatorLFO& InLFO)
 		: TModulatorProxyRefBase<FLFOId>(InLFO.GetName(), InLFO.GetUniqueID(), InLFO.bAutoActivate)
 		, Offset(InLFO.Offset)
-		, Value(0.0f)
+		, Value(1.0f)
 	{
 		LFO.SetGain(InLFO.Amplitude);
 		LFO.SetFrequency(InLFO.Frequency);
diff --git a/Engine/Source/Runtime/SignalProcessing/Private/Osc.cpp b/Engine/Source/Runtime/SignalProcessing/Private/Osc.cpp
--- a/Engine/Source/Runtime/SignalProcessing/Private/Osc.cpp
+++ b/Engine/Source/Runtime/SignalProcessing/Private/Osc.cpp
@@ -10,6 +10,7 @@ namespace Audio
 		, SampleRate(44100.0f)
 		, Nyquist(0.5f * SampleRate)
 		, Freq(440.0f)
+		, BaseFreq(1.0f)
 		, Gain(1.0f)
 		, ExternalGainMod(1.0f)
 		, Phase(0.0f)
