- Fix bounds checks
- Fix div-by-zero
- Add bypass boolean to modulation patches to avoid processing for base case

diff --git a/Engine/Plugins/Runtime/AudioModulation/Source/AudioModulation/Private/AudioModulationImpl.cpp b/Engine/Plugins/Runtime/AudioModulation/Source/AudioModulation/Private/AudioModulationImpl.cpp
--- a/Engine/Plugins/Runtime/AudioModulation/Source/AudioModulation/Private/AudioModulationImpl.cpp
+++ b/Engine/Plugins/Runtime/AudioModulation/Source/AudioModulation/Private/AudioModulationImpl.cpp
@@ -515,6 +515,11 @@ namespace AudioModulation
 
 	bool FAudioModulationImpl::CalculateModulationValue(FModulationPatchProxy& OutProxy, float& OutValue) const
 	{
+		if (OutProxy.bBypass)
+		{
+			return false;
+		}
+
 		const float InitValue = OutValue;
 		OutValue = CalculateModulationValue(OutProxy);
 
@@ -638,15 +643,49 @@ namespace AudioModulation
 
 		FModulationSettingsProxy& Settings = SourceSettings[InSourceId];
 
-		bControlsUpdated |= CalculateModulationValue(Settings.Volume, OutControls.Volume);
-		bControlsUpdated |= CalculateModulationValue(Settings.Pitch, OutControls.Pitch);
-		bControlsUpdated |= CalculateModulationValue(Settings.Lowpass, OutControls.Lowpass);
-		bControlsUpdated |= CalculateModulationValue(Settings.Highpass, OutControls.Highpass);
+		if (Settings.Volume.bBypass)
+		{
+			OutControls.Volume = 1.0f;
+		}
+		else
+		{
+			bControlsUpdated |= CalculateModulationValue(Settings.Volume, OutControls.Volume);
+		}
+
+		if (Settings.Pitch.bBypass)
+		{
+			OutControls.Pitch = 1.0f;
+		}
+		else
+		{
+			bControlsUpdated |= CalculateModulationValue(Settings.Pitch, OutControls.Pitch);
+		}
+
+		if (Settings.Highpass.bBypass)
+		{
+			OutControls.Highpass = MIN_FILTER_FREQUENCY;
+		}
+		else
+		{
+			bControlsUpdated |= CalculateModulationValue(Settings.Highpass, OutControls.Highpass);
+		}
+
+		if (Settings.Lowpass.bBypass)
+		{
+			OutControls.Lowpass = MAX_FILTER_FREQUENCY;
+		}
+		else
+		{
+			bControlsUpdated |= CalculateModulationValue(Settings.Lowpass, OutControls.Lowpass);
+		}
 
 		for (TPair<FName, FModulationPatchProxy>& Pair : Settings.Controls)
 		{
-			float& OutputValue = OutControls.Controls.FindOrAdd(Pair.Key);
-			bControlsUpdated |= CalculateModulationValue(Pair.Value, OutputValue);
+			if (!Pair.Value.bBypass)
+			{
+				float& OutputValue = OutControls.Controls.FindOrAdd(Pair.Key);
+				bControlsUpdated |= CalculateModulationValue(Pair.Value, OutputValue);
+			}
 		}
 
 		return bControlsUpdated;
diff --git a/Engine/Plugins/Runtime/AudioModulation/Source/AudioModulation/Private/SoundControlBus.cpp b/Engine/Plugins/Runtime/AudioModulation/Source/AudioModulation/Private/SoundControlBus.cpp
--- a/Engine/Plugins/Runtime/AudioModulation/Source/AudioModulation/Private/SoundControlBus.cpp
+++ b/Engine/Plugins/Runtime/AudioModulation/Source/AudioModulation/Private/SoundControlBus.cpp
@@ -40,6 +40,14 @@ void USoundControlBusBase::PostEditChangeProperty(FPropertyChangedEvent& Propert
 		{
 			Address = GetName();
 		}
+
+		if (Property->GetFName() == GET_MEMBER_NAME_CHECKED(USoundControlBusBase, DefaultValue)
+			|| Property->GetFName() == GET_MEMBER_NAME_CHECKED(USoundControlBusBase, Min)
+			|| Property->GetFName() == GET_MEMBER_NAME_CHECKED(USoundControlBusBase, Max))
+		{
+			Min = FMath::Min(Min, Max);
+			DefaultValue = FMath::Clamp(DefaultValue, Min, Max);
+		}
 	}
 
 	Super::PostEditChangeProperty(PropertyChangedEvent);
diff --git a/Engine/Plugins/Runtime/AudioModulation/Source/AudioModulation/Private/SoundModulationPatch.cpp b/Engine/Plugins/Runtime/AudioModulation/Source/AudioModulation/Private/SoundModulationPatch.cpp
--- a/Engine/Plugins/Runtime/AudioModulation/Source/AudioModulation/Private/SoundModulationPatch.cpp
+++ b/Engine/Plugins/Runtime/AudioModulation/Source/AudioModulation/Private/SoundModulationPatch.cpp
@@ -10,6 +10,21 @@
 #include "SoundModulationTransform.h"
 
 
+namespace
+{
+	template <typename T>
+	void ClampPatchInputs(TArray<T>& Inputs)
+	{
+		for (T& Input : Inputs)
+		{
+			if (Input.Transform.InputMin > Input.Transform.InputMax)
+			{
+				Input.Transform.InputMin = Input.Transform.InputMax;
+			}
+		}
+	}
+} // namespace <>
+
 USoundModulationSettings::USoundModulationSettings(const FObjectInitializer& ObjectInitializer)
 	: Super(ObjectInitializer)
 {
@@ -69,6 +84,11 @@ void FSoundModulationPatchBase::Clamp()
 {
 	if (FSoundModulationOutputBase* Output = GetOutput())
 	{
+		if (Output->Transform.InputMin > Output->Transform.InputMax)
+		{
+			Output->Transform.InputMin = Output->Transform.InputMax;
+		}
+
 		if (Output->Transform.OutputMin > Output->Transform.OutputMax)
 		{
 			Output->Transform.OutputMin = Output->Transform.OutputMax;
@@ -80,6 +100,7 @@ void FSoundVolumeModulationPatch::Clamp()
 {
 	FSoundModulationPatchBase::Clamp();
 
+	ClampPatchInputs<FSoundVolumeModulationInput>(Inputs);
 	Output.Transform.OutputMin = FMath::Clamp(Output.Transform.OutputMin, 0.0f, MAX_VOLUME);
 	Output.Transform.OutputMax = FMath::Clamp(Output.Transform.OutputMax, 0.0f, MAX_VOLUME);
 }
@@ -88,6 +109,7 @@ void FSoundPitchModulationPatch::Clamp()
 {
 	FSoundModulationPatchBase::Clamp();
 
+	ClampPatchInputs<FSoundPitchModulationInput>(Inputs);
 	Output.Transform.OutputMin = FMath::Clamp(Output.Transform.OutputMin, MIN_PITCH, MAX_PITCH);
 	Output.Transform.OutputMax = FMath::Clamp(Output.Transform.OutputMax, MIN_PITCH, MAX_PITCH);
 }
@@ -96,6 +118,7 @@ void FSoundLPFModulationPatch::Clamp()
 {
 	FSoundModulationPatchBase::Clamp();
 
+	ClampPatchInputs<FSoundLPFModulationInput>(Inputs);
 	Output.Transform.OutputMin = FMath::Clamp(Output.Transform.OutputMin, MIN_FILTER_FREQUENCY, MAX_FILTER_FREQUENCY);
 	Output.Transform.OutputMax = FMath::Clamp(Output.Transform.OutputMax, MIN_FILTER_FREQUENCY, MAX_FILTER_FREQUENCY);
 }
@@ -104,6 +127,7 @@ void FSoundHPFModulationPatch::Clamp()
 {
 	FSoundModulationPatchBase::Clamp();
 
+	ClampPatchInputs<FSoundHPFModulationInput>(Inputs);
 	Output.Transform.OutputMin = FMath::Clamp(Output.Transform.OutputMin, MIN_FILTER_FREQUENCY, MAX_FILTER_FREQUENCY);
 	Output.Transform.OutputMax = FMath::Clamp(Output.Transform.OutputMax, MIN_FILTER_FREQUENCY, MAX_FILTER_FREQUENCY);
 }
@@ -191,6 +215,7 @@ FSoundLPFModulationInput::FSoundLPFModulationInput()
 
 FSoundModulationPatchBase::FSoundModulationPatchBase()
 	: DefaultInputValue(1.0f)
+	, bBypass(1)
 {
 }
 
@@ -230,11 +255,13 @@ namespace AudioModulation
 
 	FModulationPatchProxy::FModulationPatchProxy()
 		: DefaultInputValue(1.0f)
+		, bBypass(1)
 	{
 	}
 
 	FModulationPatchProxy::FModulationPatchProxy(const FSoundModulationPatchBase& Patch)
 		: DefaultInputValue(Patch.DefaultInputValue)
+		, bBypass(Patch.bBypass)
 		, OutputProxy(*Patch.GetOutput())
 	{
 		Patch.GenerateProxies(InputProxies);
diff --git a/Engine/Plugins/Runtime/AudioModulation/Source/AudioModulation/Private/SoundModulationProxy.h b/Engine/Plugins/Runtime/AudioModulation/Source/AudioModulation/Private/SoundModulationProxy.h
--- a/Engine/Plugins/Runtime/AudioModulation/Source/AudioModulation/Private/SoundModulationProxy.h
+++ b/Engine/Plugins/Runtime/AudioModulation/Source/AudioModulation/Private/SoundModulationProxy.h
@@ -264,6 +264,9 @@ namespace AudioModulation
 		/** Default value of patch (Value mixed when inputs are provided or not, regardless of active state)*/
 		float DefaultInputValue;
 
+		/** Bypasses the patch and doesn't update modulation value */
+		uint8 bBypass : 1;
+
 		/** Optional modulation inputs */
 		TArray<FModulationInputProxy> InputProxies;
 
diff --git a/Engine/Plugins/Runtime/AudioModulation/Source/AudioModulation/Private/SoundModulationTransform.cpp b/Engine/Plugins/Runtime/AudioModulation/Source/AudioModulation/Private/SoundModulationTransform.cpp
--- a/Engine/Plugins/Runtime/AudioModulation/Source/AudioModulation/Private/SoundModulationTransform.cpp
+++ b/Engine/Plugins/Runtime/AudioModulation/Source/AudioModulation/Private/SoundModulationTransform.cpp
@@ -14,8 +14,10 @@ FSoundModulationInputTransform::FSoundModulationInputTransform()
 
 void FSoundModulationInputTransform::Apply(float& Value) const
 {
-	const float Alpha = FMath::Clamp((Value - InputMin) / (InputMax - InputMin), 0.0f, 1.0f);
+	const float Denom = FMath::Max(FMath::Abs(InputMax - InputMin), SMALL_NUMBER);
+	const float Alpha = FMath::Clamp(FMath::Abs(Value - InputMin) / Denom, 0.0f, 1.0f);
 	Value = FMath::Lerp(OutputMin, OutputMax, Alpha);
+	Value = FMath::Clamp(Value, OutputMin, OutputMax);
 }
 
 FSoundModulationOutputTransform::FSoundModulationOutputTransform()
@@ -58,13 +60,10 @@ void FSoundModulationOutputTransform::EvaluateCurve(float& Value) const
 		return;
 	}
 
-	// Avoid divide-by-zero & return untransformed if limited to single value
-	if (InputMax <= InputMin)
-	{
-		return;
-	}
-	
-	const float Alpha = (Value - InputMin) / (InputMax - InputMin);
+	// Avoid divide-by-zero
+	const float Denom = FMath::Max(FMath::Abs(InputMax - InputMin), SMALL_NUMBER);
+	const float Alpha = FMath::Clamp(FMath::Abs(Value - InputMin) / Denom, 0.0f, 1.0f);
+
 	switch (Curve)
 	{
 		case ESoundModulatorOutputCurve::Linear:
diff --git a/Engine/Plugins/Runtime/AudioModulation/Source/AudioModulation/Public/SoundModulationPatch.h b/Engine/Plugins/Runtime/AudioModulation/Source/AudioModulation/Public/SoundModulationPatch.h
--- a/Engine/Plugins/Runtime/AudioModulation/Source/AudioModulation/Public/SoundModulationPatch.h
+++ b/Engine/Plugins/Runtime/AudioModulation/Source/AudioModulation/Public/SoundModulationPatch.h
@@ -168,6 +168,10 @@ struct FSoundModulationPatchBase
 	UPROPERTY(EditAnywhere, BlueprintReadWrite, Category = Inputs, meta = (UIMin = "0", UIMax = "1"))
 	float DefaultInputValue;
 
+	/** Whether or not patch is bypassed */
+	UPROPERTY(EditAnywhere, BlueprintReadWrite, Category = Inputs, meta = (UIMin = "0", UIMax = "1"))
+	uint8 bBypass : 1;
+
 	virtual void GenerateProxies(TArray<AudioModulation::FModulationInputProxy>& InputProxies) const PURE_VIRTUAL(FSoundModulationPatchBase::GenerateProxies, );
 	virtual FSoundModulationOutputBase* GetOutput() PURE_VIRTUAL(FSoundModulationPatchBase::GetOutput(), return nullptr;);
 	virtual const FSoundModulationOutputBase* GetOutput() const PURE_VIRTUAL(FSoundModulationPatchBase::GetOutput(), return nullptr;);
diff --git a/Engine/Plugins/Runtime/AudioModulation/Source/AudioModulationEditor/Private/Layouts/SoundControlModulationPatchLayout.cpp b/Engine/Plugins/Runtime/AudioModulation/Source/AudioModulationEditor/Private/Layouts/SoundControlModulationPatchLayout.cpp
--- a/Engine/Plugins/Runtime/AudioModulation/Source/AudioModulationEditor/Private/Layouts/SoundControlModulationPatchLayout.cpp
+++ b/Engine/Plugins/Runtime/AudioModulation/Source/AudioModulationEditor/Private/Layouts/SoundControlModulationPatchLayout.cpp
@@ -45,57 +45,54 @@ void FSoundModulationPatchLayoutCustomization::CustomizeChildren(TSharedRef<IPro
 	CustomizeControl(PropertyHandles, ChildBuilder);
 }
 
-void FSoundModulationPatchLayoutCustomization::CustomizeControl(TMap<FName, TSharedPtr<IPropertyHandle>> &PropertyHandles, IDetailChildrenBuilder &ChildBuilder)
+TAttribute<EVisibility> FSoundModulationPatchLayoutCustomization::CustomizeControl(TMap<FName, TSharedPtr<IPropertyHandle>> &PropertyHandles, IDetailChildrenBuilder &ChildBuilder)
 {
-	TSharedPtr<IPropertyHandle>ValueHandle = PropertyHandles.FindChecked(GET_MEMBER_NAME_CHECKED(FSoundModulationPatchBase, DefaultInputValue));
-	ChildBuilder.AddProperty(ValueHandle.ToSharedRef());
-}
+	TSharedRef<IPropertyHandle>BypassHandle = PropertyHandles.FindChecked(GET_MEMBER_NAME_CHECKED(FSoundModulationPatchBase, bBypass)).ToSharedRef();
+	ChildBuilder.AddProperty(BypassHandle);
 
-void FSoundVolumeModulationPatchLayoutCustomization::CustomizeControl(TMap<FName, TSharedPtr<IPropertyHandle>> &PropertyHandles, IDetailChildrenBuilder &ChildBuilder)
-{
-	FSoundModulationPatchLayoutCustomization::CustomizeControl(PropertyHandles, ChildBuilder);
+	TAttribute<EVisibility> VisibilityAttribute = TAttribute<EVisibility>::Create([this, BypassHandle]()
+	{
+		bool bIsBypassed = false;
+		BypassHandle->GetValue(bIsBypassed);
+		return bIsBypassed ? EVisibility::Hidden : EVisibility::Visible;
+	});
 
-	TSharedPtr<IPropertyHandle>InputsHandle = PropertyHandles.FindChecked(GET_MEMBER_NAME_CHECKED(FSoundVolumeModulationPatch, Inputs));
-	ChildBuilder.AddProperty(InputsHandle.ToSharedRef());
+	TSharedRef<IPropertyHandle>ValueHandle = PropertyHandles.FindChecked(GET_MEMBER_NAME_CHECKED(FSoundModulationPatchBase, DefaultInputValue)).ToSharedRef();
+	ChildBuilder.AddProperty(ValueHandle)
+		.Visibility(VisibilityAttribute);
 
-	TSharedPtr<IPropertyHandle>OutputHandle = PropertyHandles.FindChecked(GET_MEMBER_NAME_CHECKED(FSoundVolumeModulationPatch, Output));
-	ChildBuilder.AddProperty(OutputHandle.ToSharedRef());
+	return VisibilityAttribute;
 }
 
-void FSoundPitchModulationPatchLayoutCustomization::CustomizeControl(TMap<FName, TSharedPtr<IPropertyHandle>> &PropertyHandles, IDetailChildrenBuilder &ChildBuilder)
+TAttribute<EVisibility> FSoundVolumeModulationPatchLayoutCustomization::CustomizeControl(TMap<FName, TSharedPtr<IPropertyHandle>> &PropertyHandles, IDetailChildrenBuilder &ChildBuilder)
 {
-	FSoundModulationPatchLayoutCustomization::CustomizeControl(PropertyHandles, ChildBuilder);
-
-	TSharedPtr<IPropertyHandle>InputsHandle = PropertyHandles.FindChecked(GET_MEMBER_NAME_CHECKED(FSoundPitchModulationPatch, Inputs));
-	ChildBuilder.AddProperty(InputsHandle.ToSharedRef());
-
-	TSharedPtr<IPropertyHandle>OutputHandle = PropertyHandles.FindChecked(GET_MEMBER_NAME_CHECKED(FSoundPitchModulationPatch, Output));
-	ChildBuilder.AddProperty(OutputHandle.ToSharedRef());
+	TAttribute<EVisibility> VisibilityAttribute = FSoundModulationPatchLayoutCustomization::CustomizeControl(PropertyHandles, ChildBuilder);
+	AddPatchProperties<FSoundVolumeModulationPatch>(VisibilityAttribute, PropertyHandles, ChildBuilder);
+	return VisibilityAttribute;
 }
 
-void FSoundHPFModulationPatchLayoutCustomization::CustomizeControl(TMap<FName, TSharedPtr<IPropertyHandle>> &PropertyHandles, IDetailChildrenBuilder &ChildBuilder)
+TAttribute<EVisibility> FSoundPitchModulationPatchLayoutCustomization::CustomizeControl(TMap<FName, TSharedPtr<IPropertyHandle>> &PropertyHandles, IDetailChildrenBuilder &ChildBuilder)
 {
-	FSoundModulationPatchLayoutCustomization::CustomizeControl(PropertyHandles, ChildBuilder);
-
-	TSharedPtr<IPropertyHandle>InputsHandle = PropertyHandles.FindChecked(GET_MEMBER_NAME_CHECKED(FSoundHPFModulationPatch, Inputs));
-	ChildBuilder.AddProperty(InputsHandle.ToSharedRef());
-
-	TSharedPtr<IPropertyHandle>OutputHandle = PropertyHandles.FindChecked(GET_MEMBER_NAME_CHECKED(FSoundHPFModulationPatch, Output));
-	ChildBuilder.AddProperty(OutputHandle.ToSharedRef());
+	TAttribute<EVisibility> VisibilityAttribute = FSoundModulationPatchLayoutCustomization::CustomizeControl(PropertyHandles, ChildBuilder);
+	AddPatchProperties<FSoundPitchModulationPatch>(VisibilityAttribute, PropertyHandles, ChildBuilder);
+	return VisibilityAttribute;
 }
 
-void FSoundLPFModulationPatchLayoutCustomization::CustomizeControl(TMap<FName, TSharedPtr<IPropertyHandle>> &PropertyHandles, IDetailChildrenBuilder &ChildBuilder)
+TAttribute<EVisibility> FSoundHPFModulationPatchLayoutCustomization::CustomizeControl(TMap<FName, TSharedPtr<IPropertyHandle>> &PropertyHandles, IDetailChildrenBuilder &ChildBuilder)
 {
-	FSoundModulationPatchLayoutCustomization::CustomizeControl(PropertyHandles, ChildBuilder);
-
-	TSharedPtr<IPropertyHandle>InputsHandle = PropertyHandles.FindChecked(GET_MEMBER_NAME_CHECKED(FSoundLPFModulationPatch, Inputs));
-	ChildBuilder.AddProperty(InputsHandle.ToSharedRef());
+	TAttribute<EVisibility> VisibilityAttribute = FSoundModulationPatchLayoutCustomization::CustomizeControl(PropertyHandles, ChildBuilder);
+	AddPatchProperties<FSoundHPFModulationPatch>(VisibilityAttribute, PropertyHandles, ChildBuilder);
+	return VisibilityAttribute;
+}
 
-	TSharedPtr<IPropertyHandle>OutputHandle = PropertyHandles.FindChecked(GET_MEMBER_NAME_CHECKED(FSoundLPFModulationPatch, Output));
-	ChildBuilder.AddProperty(OutputHandle.ToSharedRef());
+TAttribute<EVisibility> FSoundLPFModulationPatchLayoutCustomization::CustomizeControl(TMap<FName, TSharedPtr<IPropertyHandle>> &PropertyHandles, IDetailChildrenBuilder &ChildBuilder)
+{
+	TAttribute<EVisibility> VisibilityAttribute = FSoundModulationPatchLayoutCustomization::CustomizeControl(PropertyHandles, ChildBuilder);
+	AddPatchProperties<FSoundLPFModulationPatch>(VisibilityAttribute, PropertyHandles, ChildBuilder);
+	return VisibilityAttribute;
 }
 
-void FSoundControlModulationPatchLayoutCustomization::CustomizeControl(TMap<FName, TSharedPtr<IPropertyHandle>> &PropertyHandles, IDetailChildrenBuilder &ChildBuilder)
+TAttribute<EVisibility> FSoundControlModulationPatchLayoutCustomization::CustomizeControl(TMap<FName, TSharedPtr<IPropertyHandle>> &PropertyHandles, IDetailChildrenBuilder &ChildBuilder)
 {
 // Properties hidden as Generic Control Modulation is still in development
 // 	FName SetControlName;
@@ -123,5 +120,10 @@ void FSoundControlModulationPatchLayoutCustomization::CustomizeControl(TMap<FNam
 // 
 // 	TSharedPtr<IPropertyHandle>OutputHandle = PropertyHandles.FindChecked(GET_MEMBER_NAME_CHECKED(FSoundControlModulationPatch, Output));
 // 	ChildBuilder.AddProperty(OutputHandle.ToSharedRef());
+
+	return TAttribute<EVisibility>::Create([this]()
+	{
+		return EVisibility::Hidden;
+	});
 }
 #undef LOCTEXT_NAMESPACE
diff --git a/Engine/Plugins/Runtime/AudioModulation/Source/AudioModulationEditor/Private/Layouts/SoundControlModulationPatchLayout.h b/Engine/Plugins/Runtime/AudioModulation/Source/AudioModulationEditor/Private/Layouts/SoundControlModulationPatchLayout.h
--- a/Engine/Plugins/Runtime/AudioModulation/Source/AudioModulationEditor/Private/Layouts/SoundControlModulationPatchLayout.h
+++ b/Engine/Plugins/Runtime/AudioModulation/Source/AudioModulationEditor/Private/Layouts/SoundControlModulationPatchLayout.h
@@ -27,7 +27,20 @@ public:
 	virtual void CustomizeChildren(TSharedRef<IPropertyHandle> StructPropertyHandle, IDetailChildrenBuilder& ChildBuilder, IPropertyTypeCustomizationUtils& StructCustomizationUtils) override;
 	//~ End IPropertyTypeCustomization
 
-	virtual void CustomizeControl(TMap<FName, TSharedPtr<IPropertyHandle>> &PropertyHandles, IDetailChildrenBuilder &ChildBuilder);
+protected:
+	template <typename T>
+	void AddPatchProperties(TAttribute<EVisibility> VisibilityAttribute, TMap<FName, TSharedPtr<IPropertyHandle>> &PropertyHandles, IDetailChildrenBuilder &ChildBuilder)
+	{
+		TSharedPtr<IPropertyHandle>InputsHandle = PropertyHandles.FindChecked(GET_MEMBER_NAME_CHECKED(T, Inputs));
+		ChildBuilder.AddProperty(InputsHandle.ToSharedRef())
+			.Visibility(VisibilityAttribute);
+
+		TSharedPtr<IPropertyHandle>OutputHandle = PropertyHandles.FindChecked(GET_MEMBER_NAME_CHECKED(T, Output));
+		ChildBuilder.AddProperty(OutputHandle.ToSharedRef())
+			.Visibility(VisibilityAttribute);
+	}
+
+	virtual TAttribute<EVisibility> CustomizeControl(TMap<FName, TSharedPtr<IPropertyHandle>> &PropertyHandles, IDetailChildrenBuilder &ChildBuilder);
 };
 
 class FSoundVolumeModulationPatchLayoutCustomization : public FSoundModulationPatchLayoutCustomization
@@ -38,7 +51,7 @@ public:
 		return MakeShared<FSoundVolumeModulationPatchLayoutCustomization>();
 	}
 
-	virtual void CustomizeControl(TMap<FName, TSharedPtr<IPropertyHandle>> &PropertyHandles, IDetailChildrenBuilder &ChildBuilder) override;
+	virtual TAttribute<EVisibility> CustomizeControl(TMap<FName, TSharedPtr<IPropertyHandle>> &PropertyHandles, IDetailChildrenBuilder &ChildBuilder) override;
 };
 
 class FSoundPitchModulationPatchLayoutCustomization : public FSoundModulationPatchLayoutCustomization
@@ -49,7 +62,7 @@ public:
 		return MakeShared<FSoundPitchModulationPatchLayoutCustomization>();
 	}
 
-	virtual void CustomizeControl(TMap<FName, TSharedPtr<IPropertyHandle>> &PropertyHandles, IDetailChildrenBuilder &ChildBuilder) override;
+	virtual TAttribute<EVisibility> CustomizeControl(TMap<FName, TSharedPtr<IPropertyHandle>> &PropertyHandles, IDetailChildrenBuilder &ChildBuilder) override;
 };
 
 class FSoundLPFModulationPatchLayoutCustomization : public FSoundModulationPatchLayoutCustomization
@@ -60,7 +73,7 @@ public:
 		return MakeShared<FSoundLPFModulationPatchLayoutCustomization>();
 	}
 
-	virtual void CustomizeControl(TMap<FName, TSharedPtr<IPropertyHandle>> &PropertyHandles, IDetailChildrenBuilder &ChildBuilder) override;
+	virtual TAttribute<EVisibility> CustomizeControl(TMap<FName, TSharedPtr<IPropertyHandle>> &PropertyHandles, IDetailChildrenBuilder &ChildBuilder) override;
 };
 
 class FSoundHPFModulationPatchLayoutCustomization : public FSoundModulationPatchLayoutCustomization
@@ -71,7 +84,7 @@ public:
 		return MakeShared<FSoundHPFModulationPatchLayoutCustomization>();
 	}
 
-	virtual void CustomizeControl(TMap<FName, TSharedPtr<IPropertyHandle>> &PropertyHandles, IDetailChildrenBuilder &ChildBuilder) override;
+	virtual TAttribute<EVisibility> CustomizeControl(TMap<FName, TSharedPtr<IPropertyHandle>> &PropertyHandles, IDetailChildrenBuilder &ChildBuilder) override;
 };
 
 class FSoundControlModulationPatchLayoutCustomization : public FSoundModulationPatchLayoutCustomization
@@ -82,7 +95,7 @@ public:
 		return MakeShared<FSoundControlModulationPatchLayoutCustomization>();
 	}
 
-	virtual void CustomizeControl(TMap<FName, TSharedPtr<IPropertyHandle>> &PropertyHandles, IDetailChildrenBuilder &ChildBuilder) override;
+	virtual TAttribute<EVisibility> CustomizeControl(TMap<FName, TSharedPtr<IPropertyHandle>> &PropertyHandles, IDetailChildrenBuilder &ChildBuilder) override;
 
 protected:
 	TSharedPtr<SSearchableComboBox> ControlComboBox;
