VS2019 strict fix: Missing include

diff --git a/Engine/Plugins/Runtime/AudioModulation/Source/AudioModulationEditor/Private/Layouts/SoundControlModulationPatchLayout.h b/Engine/Plugins/Runtime/AudioModulation/Source/AudioModulationEditor/Private/Layouts/SoundControlModulationPatchLayout.h
--- a/Engine/Plugins/Runtime/AudioModulation/Source/AudioModulationEditor/Private/Layouts/SoundControlModulationPatchLayout.h
+++ b/Engine/Plugins/Runtime/AudioModulation/Source/AudioModulationEditor/Private/Layouts/SoundControlModulationPatchLayout.h
@@ -3,12 +3,14 @@
 
 #include "CoreMinimal.h"
 
+#include "IDetailChildrenBuilder.h"
 #include "IPropertyTypeCustomization.h"
 #include "Layout/Visibility.h"
 #include "Misc/Attribute.h"
 #include "PropertyHandle.h"
 #include "SoundModulationControlComboBox.h"
 
+
 // Forward Declarations
 class SSearchableComboBox;
 struct FSoundModulationSettings;
@@ -29,7 +31,7 @@ public:
 
 protected:
 	template <typename T>
-	void AddPatchProperties(TAttribute<EVisibility> VisibilityAttribute, TMap<FName, TSharedPtr<IPropertyHandle>> &PropertyHandles, IDetailChildrenBuilder &ChildBuilder)
+	void AddPatchProperties(TAttribute<EVisibility> VisibilityAttribute, TMap<FName, TSharedPtr<IPropertyHandle>>& PropertyHandles, IDetailChildrenBuilder& ChildBuilder)
 	{
 		TSharedPtr<IPropertyHandle>InputsHandle = PropertyHandles.FindChecked(GET_MEMBER_NAME_CHECKED(T, Inputs));
 		ChildBuilder.AddProperty(InputsHandle.ToSharedRef())
@@ -40,7 +42,7 @@ protected:
 			.Visibility(VisibilityAttribute);
 	}
 
-	virtual TAttribute<EVisibility> CustomizeControl(TMap<FName, TSharedPtr<IPropertyHandle>> &PropertyHandles, IDetailChildrenBuilder &ChildBuilder);
+	virtual TAttribute<EVisibility> CustomizeControl(TMap<FName, TSharedPtr<IPropertyHandle>>& PropertyHandles, IDetailChildrenBuilder& ChildBuilder);
 };
 
 class FSoundVolumeModulationPatchLayoutCustomization : public FSoundModulationPatchLayoutCustomization
@@ -51,7 +53,7 @@ public:
 		return MakeShared<FSoundVolumeModulationPatchLayoutCustomization>();
 	}
 
-	virtual TAttribute<EVisibility> CustomizeControl(TMap<FName, TSharedPtr<IPropertyHandle>> &PropertyHandles, IDetailChildrenBuilder &ChildBuilder) override;
+	virtual TAttribute<EVisibility> CustomizeControl(TMap<FName, TSharedPtr<IPropertyHandle>>& PropertyHandles, IDetailChildrenBuilder& ChildBuilder) override;
 };
 
 class FSoundPitchModulationPatchLayoutCustomization : public FSoundModulationPatchLayoutCustomization
@@ -62,7 +64,7 @@ public:
 		return MakeShared<FSoundPitchModulationPatchLayoutCustomization>();
 	}
 
-	virtual TAttribute<EVisibility> CustomizeControl(TMap<FName, TSharedPtr<IPropertyHandle>> &PropertyHandles, IDetailChildrenBuilder &ChildBuilder) override;
+	virtual TAttribute<EVisibility> CustomizeControl(TMap<FName, TSharedPtr<IPropertyHandle>>& PropertyHandles, IDetailChildrenBuilder& ChildBuilder) override;
 };
 
 class FSoundLPFModulationPatchLayoutCustomization : public FSoundModulationPatchLayoutCustomization
@@ -73,7 +75,7 @@ public:
 		return MakeShared<FSoundLPFModulationPatchLayoutCustomization>();
 	}
 
-	virtual TAttribute<EVisibility> CustomizeControl(TMap<FName, TSharedPtr<IPropertyHandle>> &PropertyHandles, IDetailChildrenBuilder &ChildBuilder) override;
+	virtual TAttribute<EVisibility> CustomizeControl(TMap<FName, TSharedPtr<IPropertyHandle>>& PropertyHandles, IDetailChildrenBuilder& ChildBuilder) override;
 };
 
 class FSoundHPFModulationPatchLayoutCustomization : public FSoundModulationPatchLayoutCustomization
@@ -84,7 +86,7 @@ public:
 		return MakeShared<FSoundHPFModulationPatchLayoutCustomization>();
 	}
 
-	virtual TAttribute<EVisibility> CustomizeControl(TMap<FName, TSharedPtr<IPropertyHandle>> &PropertyHandles, IDetailChildrenBuilder &ChildBuilder) override;
+	virtual TAttribute<EVisibility> CustomizeControl(TMap<FName, TSharedPtr<IPropertyHandle>>& PropertyHandles, IDetailChildrenBuilder& ChildBuilder) override;
 };
 
 class FSoundControlModulationPatchLayoutCustomization : public FSoundModulationPatchLayoutCustomization
@@ -95,7 +97,7 @@ public:
 		return MakeShared<FSoundControlModulationPatchLayoutCustomization>();
 	}
 
-	virtual TAttribute<EVisibility> CustomizeControl(TMap<FName, TSharedPtr<IPropertyHandle>> &PropertyHandles, IDetailChildrenBuilder &ChildBuilder) override;
+	virtual TAttribute<EVisibility> CustomizeControl(TMap<FName, TSharedPtr<IPropertyHandle>>& PropertyHandles, IDetailChildrenBuilder& ChildBuilder) override;
 
 protected:
 	TSharedPtr<SSearchableComboBox> ControlComboBox;
