Fix tool Default Input Value tool-tip

diff --git a/Engine/Plugins/Runtime/AudioModulation/Source/AudioModulation/Public/SoundModulationPatch.h b/Engine/Plugins/Runtime/AudioModulation/Source/AudioModulation/Public/SoundModulationPatch.h
--- a/Engine/Plugins/Runtime/AudioModulation/Source/AudioModulation/Public/SoundModulationPatch.h
+++ b/Engine/Plugins/Runtime/AudioModulation/Source/AudioModulation/Public/SoundModulationPatch.h
@@ -164,7 +164,7 @@ struct FSoundModulationPatchBase
 
 	virtual ~FSoundModulationPatchBase() = default;
 
-	/** Default value of patch (value used when all inputs are either not provided or not active) */
+	/** Default value of patch, included in mix calculation regardless of number of active buses referenced. */
 	UPROPERTY(EditAnywhere, BlueprintReadWrite, Category = Inputs, meta = (UIMin = "0", UIMax = "1"))
 	float DefaultInputValue;
 
