Fix trivial missing initializer warning in cooked builds with modulation plugin enabled

diff --git a/Engine/Plugins/Runtime/AudioModulation/Source/AudioModulation/Private/SoundModulationTransform.cpp b/Engine/Plugins/Runtime/AudioModulation/Source/AudioModulation/Private/SoundModulationTransform.cpp
--- a/Engine/Plugins/Runtime/AudioModulation/Source/AudioModulation/Private/SoundModulationTransform.cpp
+++ b/Engine/Plugins/Runtime/AudioModulation/Source/AudioModulation/Private/SoundModulationTransform.cpp
@@ -25,6 +25,7 @@ FSoundModulationOutputTransform::FSoundModulationOutputTransform()
 	, InputMax(1.0f)
 	, Curve(ESoundModulatorOutputCurve::Linear)
 	, Scalar(2.5f)
+	, CurveShared(nullptr)
 	, OutputMin(0.0f)
 	, OutputMax(1.0f)
 {
