Add bypass & handle system to fix reference bugs in release
- Includes deprecating auto-activation in favor of bypass as proxies are now always active when referenced.

diff --git a/Engine/Plugins/Runtime/AudioModulation/Source/AudioModulation/Private/AudioModulation.cpp b/Engine/Plugins/Runtime/AudioModulation/Source/AudioModulation/Private/AudioModulation.cpp
--- a/Engine/Plugins/Runtime/AudioModulation/Source/AudioModulation/Private/AudioModulation.cpp
+++ b/Engine/Plugins/Runtime/AudioModulation/Source/AudioModulation/Private/AudioModulation.cpp
@@ -26,7 +26,7 @@ namespace AudioModulation
 		return Impl.Get();
 	}
 
-	float FAudioModulation::CalculateInitialVolume(const USoundModulationPluginSourceSettingsBase& Settings) const
+	float FAudioModulation::CalculateInitialVolume(const USoundModulationPluginSourceSettingsBase& Settings)
 	{
 		return Impl->CalculateInitialVolume(Settings);
 	}
diff --git a/Engine/Plugins/Runtime/AudioModulation/Source/AudioModulation/Private/AudioModulationDebugger.cpp b/Engine/Plugins/Runtime/AudioModulation/Source/AudioModulation/Private/AudioModulationDebugger.cpp
--- a/Engine/Plugins/Runtime/AudioModulation/Source/AudioModulation/Private/AudioModulationDebugger.cpp
+++ b/Engine/Plugins/Runtime/AudioModulation/Source/AudioModulation/Private/AudioModulationDebugger.cpp
@@ -11,6 +11,7 @@
 #include "IAudioExtensionPlugin.h"
 #include "Misc/CoreDelegates.h"
 #include "SoundModulationPatch.h"
+#include "SoundModulationProxy.h"
 #include "SoundModulationValue.h"
 
 
@@ -28,7 +29,7 @@ namespace AudioModulation
 				: FColor::Green;
 		}
 
-		int32 RenderStatLFO(const TArray<FModulatorLFOProxy>& FilteredLFOs, FCanvas& Canvas, int32 X, int32 Y, const UFont& Font)
+		int32 RenderStatLFO(const TArray<FLFODebugInfo>& FilteredLFOs, FCanvas& Canvas, int32 X, int32 Y, const UFont& Font)
 		{
 			int32 Height = 12;
 			int32 Width = 12;
@@ -50,9 +51,9 @@ namespace AudioModulation
 			};
 
 			int32 CellWidth = FMath::Max(StaticCellWidths);
-			for (const FModulatorLFOProxy& LFOProxy : FilteredLFOs)
+			for (const FLFODebugInfo& LFODebugInfo : FilteredLFOs)
 			{
-				CellWidth = FMath::Max(CellWidth, LFOProxy.GetName().Len());
+				CellWidth = FMath::Max(CellWidth, LFODebugInfo.Name.Len());
 			}
 
 			Canvas.DrawShadowedString(X + XIndent, Y, TEXT("Active LFOs:"), &Font, FColor::Red);
@@ -76,17 +77,19 @@ namespace AudioModulation
 			RowX += Width * CellWidth;
 
 			Y += Height;
-			for (const FModulatorLFOProxy& LFOPair : FilteredLFOs)
+			for (const FLFODebugInfo& LFODebugInfo : FilteredLFOs)
 			{
 				RowX = X;
 
-				FString Name = LFOPair.GetName().Left(MaxNameLength);
+				FString Name = LFODebugInfo.Name.Left(MaxNameLength);
+				Name += FString::Printf(TEXT(" (%u)"), LFODebugInfo.RefCount);
+
 				if (Name.Len() < MaxNameLength)
 				{
 					Name = Name.RightPad(MaxNameLength - Name.Len());
 				}
 
-				const FString ValueString  = FString::Printf(TEXT("%.6f"), LFOPair.GetValue());
+				const FString ValueString  = FString::Printf(TEXT("%.6f"), LFODebugInfo.Value);
 
 				Canvas.DrawShadowedString(RowX, Y, *Name, &Font, FColor::Green);
 				RowX += Width * CellWidth;
@@ -100,7 +103,7 @@ namespace AudioModulation
 			return Y;
 		}
 
-		int32 RenderStatMixMatrix(const TArray<FModulatorBusMixProxy>& FilteredMixes, const TArray<FControlBusProxy>& FilteredBuses, FCanvas& Canvas, int32 X, int32 Y, const UFont& Font)
+		int32 RenderStatMixMatrix(const TArray<FControlBusMixDebugInfo>& FilteredMixes, const TArray<FControlBusDebugInfo>& FilteredBuses, FCanvas& Canvas, int32 X, int32 Y, const UFont& Font)
 		{
 			int32 Height = 12;
 			int32 Width = 12;
@@ -123,20 +126,22 @@ namespace AudioModulation
 			};
 
 			int32 CellWidth = FMath::Max(StaticCellWidths);
-			for (const FModulatorBusMixProxy& BusMix : FilteredMixes)
+			for (const FControlBusMixDebugInfo& BusMix : FilteredMixes)
 			{
-				CellWidth = FMath::Max(CellWidth, BusMix.GetName().Len());
+				CellWidth = FMath::Max(CellWidth, BusMix.Name.Len());
 			}
-			for (const FControlBusProxy& BusPair : FilteredBuses)
+			for (const FControlBusDebugInfo& Bus : FilteredBuses)
 			{
-				CellWidth = FMath::Max(CellWidth, BusPair.GetName().Len());
+				CellWidth = FMath::Max(CellWidth, Bus.Name.Len());
 			}
 
 			// Draw Column Headers
 			int32 RowX = X;
-			for (const FControlBusProxy& BusPair : FilteredBuses)
+			for (const FControlBusDebugInfo& Bus : FilteredBuses)
 			{
-				FString Name = BusPair.GetName().Left(MaxNameLength);
+				FString Name = Bus.Name.Left(MaxNameLength);
+				Name += FString::Printf(TEXT(" (%u)"), Bus.RefCount);
+
 				if (Name.Len() < MaxNameLength)
 				{
 					Name = Name.RightPad(MaxNameLength - Name.Len());
@@ -148,9 +153,11 @@ namespace AudioModulation
 
 			// Draw Row Headers
 			int32 ColumnY = Y;
-			for (const FModulatorBusMixProxy& BusMix : FilteredMixes)
+			for (const FControlBusMixDebugInfo& BusMix : FilteredMixes)
 			{
-				FString Name = BusMix.GetName().Left(MaxNameLength);
+				FString Name = BusMix.Name.Left(MaxNameLength);
+				Name += FString::Printf(TEXT(" (%u)"), BusMix.RefCount);
+
 				if (Name.Len() < MaxNameLength)
 				{
 					Name = Name.RightPad(MaxNameLength - Name.Len());
@@ -163,21 +170,21 @@ namespace AudioModulation
 			// Reset Corner of Matrix & Draw Per Bus Data
 			RowX = X;
 			ColumnY = Y;
-			for (const FModulatorBusMixProxy& BusMix : FilteredMixes)
+			for (const FControlBusMixDebugInfo& BusMix : FilteredMixes)
 			{
 				ColumnY += Height; // Add before to leave space for column headers
 				RowX = X;
 
-				for (const FControlBusProxy& Bus : FilteredBuses)
+				for (const FControlBusDebugInfo& Bus : FilteredBuses)
 				{
 					RowX += Width * CellWidth; // Add before to leave space for row headers
 
-					float Target = Bus.GetDefaultValue();
-					float Value = Bus.GetDefaultValue();
-					if (const FModulatorBusMixChannelProxy* ChannelProxy = BusMix.Channels.Find(Bus.GetId()))
+					float Target = Bus.DefaultValue;
+					float Value = Bus.DefaultValue;
+					if (const FControlBusMixChannelDebugInfo* Channel = BusMix.Channels.Find(Bus.Id))
 					{
-						Target = ChannelProxy->Value.TargetValue;
-						Value  = ChannelProxy->Value.GetCurrentValue();
+						Target = Channel->TargetValue;
+						Value  = Channel->CurrentValue;
 					}
 
 					if (Target != Value)
@@ -196,18 +203,18 @@ namespace AudioModulation
 			// Draw Sub-Totals & Totals
 			Canvas.DrawShadowedString(X, Y, *MixSubTotalHeader, &Font, FColor::Yellow);
 			RowX = X;
-			for (const FControlBusProxy& BusPair : FilteredBuses)
+			for (const FControlBusDebugInfo& Bus : FilteredBuses)
 			{
 				RowX += Width * CellWidth; // Add before to leave space for row headers
 
-				const float Value = BusPair.GetMixValue();
+				const float Value = Bus.MixValue;
 				if (FMath::IsNaN(Value))
 				{
 					Canvas.DrawShadowedString(RowX, Y, TEXT("N/A"), &Font, FColor::Green);
 				}
 				else
 				{
-					const FColor Color = GetUnitRenderColor(Value, BusPair.GetRange());
+					const FColor Color = GetUnitRenderColor(Value, Bus.Range);
 					Canvas.DrawShadowedString(RowX, Y, *FString::Printf(TEXT("%.4f"), Value), &Font, Color);
 				}
 			}
@@ -215,23 +222,22 @@ namespace AudioModulation
 
 			Canvas.DrawShadowedString(X, Y, *LFOSubTotalHeader, &Font, FColor::Yellow);
 			RowX = X;
-			for (const FControlBusProxy& BusPair : FilteredBuses)
+			for (const FControlBusDebugInfo& Bus : FilteredBuses)
 			{
 				RowX += Width * CellWidth; // Add before to leave space for row headers
-				const float Value = BusPair.GetLFOValue();
-				const FColor Color = GetUnitRenderColor(Value, BusPair.GetRange());
+				const float Value = Bus.LFOValue;
+				const FColor Color = GetUnitRenderColor(Value, Bus.Range);
 				Canvas.DrawShadowedString(RowX, Y, *FString::Printf(TEXT("%.4f"), Value), &Font, Color);
 			}
 
 			Y += Height;
 			Canvas.DrawShadowedString(X, Y, *TotalHeader, &Font, FColor::Yellow);
 			RowX = X;
-			for (const FControlBusProxy& BusPair : FilteredBuses)
+			for (const FControlBusDebugInfo& Bus : FilteredBuses)
 			{
 				RowX += Width * CellWidth; // Add before to leave space for row headers
-				const float Value = BusPair.GetValue();
-				const FColor Color = GetUnitRenderColor(Value, BusPair.GetRange());
-				Canvas.DrawShadowedString(RowX, Y, *FString::Printf(TEXT("%.4f"), Value), &Font, Color);
+				const FColor Color = GetUnitRenderColor(Bus.Value, Bus.Range);
+				Canvas.DrawShadowedString(RowX, Y, *FString::Printf(TEXT("%.4f"), Bus.Value), &Font, Color);
 			}
 			Y += Height;
 
@@ -241,14 +247,12 @@ namespace AudioModulation
 		template <typename T>
 		bool CompareNames(const T& A, const T& B)
 		{
-			return A.GetName() < B.GetName();
+			return A.Name < B.Name;
 		}
 
 		template <typename T, typename U>
-		void FilterDebugArray(const T& Map, const FString& FilterString, int32 MaxCount, TArray<U>& FilteredArray)
+		void FilterDebugArray(const T& Map, const FString& FilterString, int32 MaxCount, TArray<const U*>& FilteredArray)
 		{
-			FilteredArray.Reset(MaxCount);
-
 			int32 FilteredItemCount = 0;
 			for (const auto& IdItemPair : Map)
 			{
@@ -257,7 +261,7 @@ namespace AudioModulation
 					&& !Item.GetName().Contains(FilterString);
 				if (!Filtered)
 				{
-					FilteredArray.Add(Item);
+					FilteredArray.Add(&Item);
 					if (++FilteredItemCount >= MaxCount)
 					{
 						return;
@@ -274,10 +278,7 @@ namespace AudioModulation
 	{
 	}
 
-	void FAudioModulationDebugger::UpdateDebugData(
-		const BusProxyMap&    ActiveBuses,
-		const BusMixProxyMap& ActiveMixes,
-		const LFOProxyMap&    ActiveLFOs)
+	void FAudioModulationDebugger::UpdateDebugData(const FReferencedProxies& RefProxies)
 	{
 		check(IsInAudioThread());
 
@@ -287,27 +288,66 @@ namespace AudioModulation
 		}
 
 		static const int32 MaxFilteredBuses = 8;
-		TArray<FControlBusProxy> InFilteredBuses;
-		FilterDebugArray<BusProxyMap, FControlBusProxy>(ActiveBuses, BusStringFilter, MaxFilteredBuses, InFilteredBuses);
+		TArray<const FControlBusProxy*> FilteredBusProxies;
+		FilterDebugArray<FBusProxyMap, FControlBusProxy>(RefProxies.Buses, BusStringFilter, MaxFilteredBuses, FilteredBusProxies);
+		TArray<FControlBusDebugInfo> RefreshedFilteredBuses;
+		for (const FControlBusProxy* Proxy : FilteredBusProxies)
+		{
+			FControlBusDebugInfo DebugInfo;
+			DebugInfo.DefaultValue = Proxy->GetDefaultValue();
+			DebugInfo.Id = Proxy->GetId();
+			DebugInfo.LFOValue = Proxy->GetLFOValue();
+			DebugInfo.MixValue = Proxy->GetMixValue();
+			DebugInfo.Name = Proxy->GetName();
+			DebugInfo.Range = Proxy->GetRange();
+			DebugInfo.RefCount = Proxy->GetRefCount();
+			DebugInfo.Value = Proxy->GetValue();
+			RefreshedFilteredBuses.Add(DebugInfo);
+		}
 
 		static const int32 MaxFilteredMixes = 16;
-		TArray<FModulatorBusMixProxy> InFilteredMixes;
-		FilterDebugArray<BusMixProxyMap, FModulatorBusMixProxy>(ActiveMixes, MixStringFilter, MaxFilteredMixes, InFilteredMixes);
+		TArray<const FModulatorBusMixProxy*> FilteredMixProxies;
+		FilterDebugArray<FBusMixProxyMap, FModulatorBusMixProxy>(RefProxies.BusMixes, MixStringFilter, MaxFilteredMixes, FilteredMixProxies);
+		TArray<FControlBusMixDebugInfo> RefreshedFilteredMixes;
+		for (const FModulatorBusMixProxy* Proxy : FilteredMixProxies)
+		{
+			FControlBusMixDebugInfo DebugInfo;
+			DebugInfo.Name = Proxy->GetName();
+			DebugInfo.RefCount = Proxy->GetRefCount();
+			for (const TPair< FBusId, FModulatorBusMixChannelProxy>& Channel : Proxy->Channels)
+			{
+				FControlBusMixChannelDebugInfo ChannelDebugInfo;
+				ChannelDebugInfo.CurrentValue = Channel.Value.Value.GetCurrentValue();
+				ChannelDebugInfo.TargetValue = Channel.Value.Value.TargetValue;
+				DebugInfo.Channels.Add(Channel.Key, ChannelDebugInfo);
+			}
+			RefreshedFilteredMixes.Add(DebugInfo);
+		}
 
 		static const int32 MaxFilteredLFOs = 8;
-		TArray<FModulatorLFOProxy> InFilteredLFOs;
-		FilterDebugArray<LFOProxyMap, FModulatorLFOProxy>(ActiveLFOs, LFOStringFilter, MaxFilteredLFOs, InFilteredLFOs);
+		TArray<const FModulatorLFOProxy*> FilteredLFOProxies;
+		FilterDebugArray<FLFOProxyMap, FModulatorLFOProxy>(RefProxies.LFOs, LFOStringFilter, MaxFilteredLFOs, FilteredLFOProxies);
+
+		TArray<FLFODebugInfo> RefreshedFilteredLFOs;
+		for (const FModulatorLFOProxy* Proxy : FilteredLFOProxies)
+		{
+			FLFODebugInfo DebugInfo;
+			DebugInfo.Name = Proxy->GetName();
+			DebugInfo.RefCount = Proxy->GetRefCount();
+			DebugInfo.Value = Proxy->GetValue();
+			RefreshedFilteredLFOs.Add(DebugInfo);
+		}
 
-		FAudioThread::RunCommandOnGameThread([this, InFilteredBuses, InFilteredMixes, InFilteredLFOs]()
+		FAudioThread::RunCommandOnGameThread([this, RefreshedFilteredBuses, RefreshedFilteredMixes, RefreshedFilteredLFOs]()
 		{
-			FilteredBuses = InFilteredBuses;
-			FilteredBuses.Sort(&CompareNames<FControlBusProxy>);
+			FilteredBuses = RefreshedFilteredBuses;
+			FilteredBuses.Sort(&CompareNames<FControlBusDebugInfo>);
 
-			FilteredMixes = InFilteredMixes;
-			FilteredMixes.Sort(&CompareNames<FModulatorBusMixProxy>);
+			FilteredMixes = RefreshedFilteredMixes;
+			FilteredMixes.Sort(&CompareNames<FControlBusMixDebugInfo>);
 
-			FilteredLFOs = InFilteredLFOs;
-			FilteredLFOs.Sort(&CompareNames<FModulatorLFOProxy>);
+			FilteredLFOs = RefreshedFilteredLFOs;
+			FilteredLFOs.Sort(&CompareNames<FLFODebugInfo>);
 		});
 	}
 
diff --git a/Engine/Plugins/Runtime/AudioModulation/Source/AudioModulation/Private/AudioModulationDebugger.h b/Engine/Plugins/Runtime/AudioModulation/Source/AudioModulation/Private/AudioModulationDebugger.h
--- a/Engine/Plugins/Runtime/AudioModulation/Source/AudioModulation/Private/AudioModulationDebugger.h
+++ b/Engine/Plugins/Runtime/AudioModulation/Source/AudioModulation/Private/AudioModulationDebugger.h
@@ -13,15 +13,49 @@
 
 namespace AudioModulation
 {
+	// Forward Declarations
+	struct FReferencedProxies;
+
+	struct FControlBusMixChannelDebugInfo
+	{
+		float TargetValue;
+		float CurrentValue;
+	};
+
+	struct FControlBusMixDebugInfo
+	{
+		FString Name;
+		uint32 Id;
+		uint32 RefCount;
+
+		TMap<uint32, FControlBusMixChannelDebugInfo> Channels;
+	};
+
+	struct FControlBusDebugInfo
+	{
+		FString Name;
+		float DefaultValue;
+		float LFOValue;
+		float MixValue;
+		float Value;
+		FVector2D Range;
+		uint32 Id;
+		uint32 RefCount;
+	};
+
+	struct FLFODebugInfo
+	{
+		FString Name;
+		float Value;
+		uint32 RefCount;
+	};
+
 	class FAudioModulationDebugger
 	{
 	public:
 		FAudioModulationDebugger();
 
-		void UpdateDebugData(
-			const BusProxyMap&    ActiveBuses,
-			const BusMixProxyMap& ActiveMixes,
-			const LFOProxyMap&    ActiveLFOs);
+		void UpdateDebugData(const FReferencedProxies& RefProxies);
 		bool OnPostHelp(FCommonViewportClient& ViewportClient, const TCHAR* Stream);
 		int32 OnRenderStat(FCanvas& Canvas, int32 X, int32 Y, const UFont& Font);
 		bool OnToggleStat(FCommonViewportClient& ViewportClient, const TCHAR* Stream);
@@ -31,9 +65,9 @@ namespace AudioModulation
 		uint8 bShowRenderStatLFO : 1;
 		uint8 bShowRenderStatMix : 1;
 
-		TArray<FControlBusProxy>    FilteredBuses;
-		TArray<FModulatorLFOProxy>    FilteredLFOs;
-		TArray<FModulatorBusMixProxy> FilteredMixes;
+		TArray<FControlBusDebugInfo>    FilteredBuses;
+		TArray<FLFODebugInfo>           FilteredLFOs;
+		TArray<FControlBusMixDebugInfo> FilteredMixes;
 
 		FString BusStringFilter;
 		FString LFOStringFilter;
diff --git a/Engine/Plugins/Runtime/AudioModulation/Source/AudioModulation/Private/AudioModulationImpl.cpp b/Engine/Plugins/Runtime/AudioModulation/Source/AudioModulation/Private/AudioModulationImpl.cpp
--- a/Engine/Plugins/Runtime/AudioModulation/Source/AudioModulation/Private/AudioModulationImpl.cpp
+++ b/Engine/Plugins/Runtime/AudioModulation/Source/AudioModulation/Private/AudioModulationImpl.cpp
@@ -9,7 +9,9 @@
 #include "Misc/CoreDelegates.h"
 #include "SoundModulatorLFO.h"
 #include "SoundModulationPatch.h"
+#include "SoundModulationProxy.h"
 #include "SoundModulationValue.h"
+#include "UObject/WeakObjectPtr.h"
 
 #if !UE_BUILD_SHIPPING
 #include "AudioModulationDebugger.h"
@@ -22,40 +24,6 @@ DECLARE_DWORD_COUNTER_STAT(TEXT("Mix Count"), STAT_AudioModulationMixCount, STAT
 
 namespace
 {
-#if WITH_EDITOR
-	// Checks whether new inputs being modified require stopping referencing preview sound for proxies to be updated and take change.
-	template <typename T>
-	bool InputUpdateRequiresStop(const TArray<T>& NewInputs, const TArray<AudioModulation::FModulationInputProxy>& CurrentInputProxies)
-	{
-		check(IsInAudioThread());
-
-		if (NewInputs.Num() != CurrentInputProxies.Num())
-		{
-			return true;
-		}
-
-		for (int32 i = 0; i < NewInputs.Num(); ++i)
-		{
-			if (const USoundControlBusBase* NewBus = NewInputs[i].GetBus())
-			{
-				if (NewBus->GetUniqueID() != CurrentInputProxies[i].BusId)
-				{
-					return true;
-				}
-			}
-			else
-			{
-				if (CurrentInputProxies[i].BusId != INDEX_NONE)
-				{
-					return true;
-				}
-			}
-		}
-
-		return false;
-	};
-#endif
-
 	void MixInModulationValue(ESoundModulatorOperator& Operator, float ModStageValue, float& Value)
 	{
 		switch (Operator)
@@ -87,9 +55,6 @@ namespace
 namespace AudioModulation
 {
 	FAudioModulationImpl::FAudioModulationImpl()
-#if WITH_EDITOR
-		: PreviewSound(nullptr)
-#endif // WITH_EDITOR
 	{
 	}
 
@@ -99,101 +64,31 @@ namespace AudioModulation
 	}
 
 #if WITH_EDITOR
-	void FAudioModulationImpl::EndPreviewSession()
-	{
-		PreviewSound = nullptr;
-		PreviewSettings = FModulationSettingsProxy();
-
-		SoundSettings.Reset();
-
-		ActiveBuses.Reset();
-		ActiveBusMixes.Reset();
-		ActiveLFOs.Reset();
-
-		UE_LOG(LogAudioModulation, Verbose, TEXT("Releasing all modulation proxies: Sound preview session has ended."));
-	}
-
 	void FAudioModulationImpl::OnEditPluginSettings(const USoundModulationPluginSourceSettingsBase& InSettings)
 	{
-		// Find if sound is being referenced and auditioned and stop immediately.
-		// This informs user that modifying sound's settings does not translate to
-		// currently playing sound.
-		const uint32 SettingsId = InSettings.GetUniqueID();
-		RunCommandOnAudioThread([this, SettingsId]()
+		const TWeakObjectPtr<const USoundModulationPluginSourceSettingsBase> SettingsPtr = &InSettings;
+		RunCommandOnAudioThread([this, SettingsPtr]()
 		{
-			if (!PreviewSound)
+			if (!SettingsPtr.IsValid())
 			{
 				return;
 			}
 
-			USoundModulationPluginSourceSettingsBase* SettingsBase = PreviewSound->FindModulationSettings();
-			if (!SettingsBase || SettingsId != SettingsBase->GetUniqueID())
-			{
-				return;
-			}
-
-			const USoundModulationSettings* Settings = CastChecked<USoundModulationSettings>(SettingsBase);
-
-			bool bShouldStop = false;
-			if (InputUpdateRequiresStop<FSoundVolumeModulationInput>(Settings->Volume.Inputs, PreviewSettings.Volume.InputProxies))
+			const USoundModulationSettings* Settings = CastChecked<USoundModulationSettings>(SettingsPtr.Get());
+			const uint32 SettingsId = Settings->GetUniqueID();
+			for (FModulationSettingsProxy& SourceSetting : SourceSettings)
 			{
-				bShouldStop = true;
-			}
-			else if (InputUpdateRequiresStop<FSoundHPFModulationInput>(Settings->Highpass.Inputs, PreviewSettings.Highpass.InputProxies))
-			{
-				bShouldStop = true;
-			}
-			else if (InputUpdateRequiresStop<FSoundLPFModulationInput>(Settings->Lowpass.Inputs, PreviewSettings.Lowpass.InputProxies))
-			{
-				bShouldStop = true;
-			}
-			else if (InputUpdateRequiresStop<FSoundPitchModulationInput>(Settings->Pitch.Inputs, PreviewSettings.Pitch.InputProxies))
-			{
-				bShouldStop = true;
-			}
-			else
-			{
-				for (int32 i = 0; i < Settings->Controls.Num(); ++i)
+				if (SourceSetting.GetId() == SettingsId)
 				{
-					const FSoundControlModulationPatch& Control = Settings->Controls[i];
-					if (const FModulationPatchProxy* PatchProxy = PreviewSettings.Controls.Find(Control.Control))
-					{
-						if (InputUpdateRequiresStop<FSoundControlModulationInput>(Control.Inputs, PatchProxy->InputProxies))
-						{
-							bShouldStop = true;
-							break;
-						}
-					}
-					else
-					{
-						bShouldStop = true;
-						break;
-					}
+					SourceSetting = FModulationSettingsProxy(*Settings, RefProxies);
 				}
 			}
 
-			if (bShouldStop)
-			{
-				PreviewSound->Stop();
-				EndPreviewSession();
-			}
-			else
+			for (TPair<uint32, FModulationSettingsProxy>& Pair : SoundSettings)
 			{
-				PreviewSettings = FModulationSettingsProxy(*Settings);
-				for (FModulationSettingsProxy& SourceSetting : SourceSettings)
+				if (Pair.Value.GetId() == SettingsId)
 				{
-					if (SourceSetting.GetId() == PreviewSettings.GetId())
-					{
-						SourceSetting = PreviewSettings;
-					}
-				}
-
-				for (TPair<uint32, FModulationSettingsProxy>& Pair : SoundSettings)
-				{
-					if (Pair.Value.GetId() == PreviewSettings.GetId())
-					{
-						Pair.Value = PreviewSettings;
-					}
+					Pair.Value = FModulationSettingsProxy(*Settings, RefProxies);
 				}
 			}
 		});
@@ -204,75 +99,11 @@ namespace AudioModulation
 	{
 		check(IsInAudioThread());
 
-		const USoundModulationSettings* Settings = CastChecked<USoundModulationSettings>(&InSettings);
-
-#if WITH_EDITOR
-		if (InSound.IsPreviewSound())
-		{
-			EndPreviewSession();
-			PreviewSound = &InSound;
-			PreviewSettings = FModulationSettingsProxy(*Settings);
-		}
-		// If currently previewing, activation/deactivation is by-passing
-		// auto-activation system and always activating/deactivating,
-		// so ignore requests from non-preview to avoid mixed state.
-		else if (PreviewSound)
-		{
-			return;
-		}
-#endif // WITH_EDITOR
-
 		const uint32 SoundId = InSound.GetObjectId();
 		if (!SoundSettings.Contains(SoundId))
 		{
-			SoundSettings.Add(SoundId, FModulationSettingsProxy(*Settings));
-		}
-
-		for (const USoundControlBusMix* Mix : Settings->Mixes)
-		{
-			ActivateBusMix(*Mix, &InSound);
-		}
-
-		auto CheckRefActivate = [this, &InSound](const USoundControlBusBase* Bus)
-		{
-			if (!Bus)
-			{
-				return;
-			}
-
-			ActivateBus(*Bus, &InSound);
-
-			for (const USoundBusModulatorBase* Modulator : Bus->Modulators)
-			{
-				if (const USoundBusModulatorLFO* LFO = Cast<USoundBusModulatorLFO>(Modulator))
-				{
-					FModulatorLFOProxy LFOProxy(*LFO);
-					ActivateLFO(*LFO, &InSound);
-				}
-			}
-		};
-
-		for (const FSoundVolumeModulationInput& Input : Settings->Volume.Inputs)
-		{
-			CheckRefActivate(Cast<USoundControlBusBase>(Input.Bus));
-		}
-
-		for (const FSoundPitchModulationInput& Input : Settings->Pitch.Inputs)
-		{
-			CheckRefActivate(Cast<USoundControlBusBase>(Input.Bus));
-		}
-
-		for (const FSoundLPFModulationInput& Input : Settings->Lowpass.Inputs)
-		{
-			CheckRefActivate(Cast<USoundControlBusBase>(Input.Bus));
-		}
-
-		for (const FSoundControlModulationPatch& Patch : Settings->Controls)
-		{
-			for (const FSoundControlModulationInput& Input : Patch.Inputs)
-			{
-				CheckRefActivate(Cast<USoundControlBusBase>(Input.Bus));
-			}
+			const USoundModulationSettings* Settings = CastChecked<USoundModulationSettings>(&InSettings);
+			SoundSettings.Add(SoundId, FModulationSettingsProxy(*Settings, RefProxies));
 		}
 	}
 
@@ -282,7 +113,7 @@ namespace AudioModulation
 
 		if (const USoundModulationSettings* Settings = CastChecked<USoundModulationSettings>(&InSettings))
 		{
-			SourceSettings[InSourceId] = FModulationSettingsProxy(*Settings);
+			SourceSettings[InSourceId] = FModulationSettingsProxy(*Settings, RefProxies);
 		}
 	}
 
@@ -298,60 +129,9 @@ namespace AudioModulation
 		check(IsInAudioThread());
 		check(InSound.GetObjectId() != INDEX_NONE);
 
-#if WITH_EDITOR
-		// If previewing, don't attempt to deactivate anything, leaving
-		// all modulation proxies active until end of session, at which
-		// point all will be reset.
-		if (PreviewSound)
-		{
-			if (&InSound == PreviewSound)
-			{
-				EndPreviewSession();
-			}
-			return;
-		}
-#endif // WITH_EDITOR
-
-		// Settings can be null if sound settings were modified via the editor while auditioning or in PIE
-		if (const FModulationSettingsProxy* Settings = SoundSettings.Find(InSound.GetObjectId()))
+		if (InSound.GetPlayCount() == 0)
 		{
-			for (FBusMixId MixId : Settings->Mixes)
-			{
-				DeactivateBusMix(MixId, &InSound);
-			}
-
-			for (const FModulationInputProxy& Input : Settings->Volume.InputProxies)
-			{
-				DeactivateBus(Input.BusId, &InSound);
-			}
-
-			for (const FModulationInputProxy& Input : Settings->Pitch.InputProxies)
-			{
-				DeactivateBus(Input.BusId, &InSound);
-			}
-
-			for (const FModulationInputProxy& Input : Settings->Lowpass.InputProxies)
-			{
-				DeactivateBus(Input.BusId, &InSound);
-			}
-
-			for (const FModulationInputProxy& Input : Settings->Highpass.InputProxies)
-			{
-				DeactivateBus(Input.BusId, &InSound);
-			}
-
-			for (const TPair<FName, FModulationPatchProxy>& PatchProxy : Settings->Controls)
-			{
-				for (const FModulationInputProxy& Input : PatchProxy.Value.InputProxies)
-				{
-					DeactivateBus(Input.BusId, &InSound);
-				}
-			}
-
-			if (InSound.GetPlayCount() == 0)
-			{
-				SoundSettings.Remove(InSound.GetObjectId());
-			}
+			SoundSettings.Remove(InSound.GetObjectId());
 		}
 	}
 
@@ -372,70 +152,48 @@ namespace AudioModulation
 	}
 #endif // !UE_BUILD_SHIPPING
 
-	void FAudioModulationImpl::ActivateBus(const USoundControlBusBase& InBus, const ISoundModulatable* InSound)
+	void FAudioModulationImpl::ActivateBus(const USoundControlBusBase& InBus)
 	{
-		const FControlBusProxy NewBusProxy(InBus);
-		const bool bCanCreateNew = InBus.CanAutoActivate(InSound);
-
-		RunCommandOnAudioThread([this, NewBusProxy, InSound, bCanCreateNew]()
+		const TWeakObjectPtr<const USoundControlBusBase>BusPtr = &InBus;
+		RunCommandOnAudioThread([this, BusPtr]()
 		{
-			const FBusId NewBusId = NewBusProxy.GetId();
-			FControlBusProxy* BusProxy = ActiveBuses.Find(NewBusId);
-			if (!BusProxy && bCanCreateNew)
+			if (BusPtr.IsValid())
 			{
-				BusProxy = &ActiveBuses.Add(NewBusId, NewBusProxy);
-			}
+				auto OnCreate = [this, BusPtr](FControlBusProxy& NewProxy)
+				{
+					NewProxy.InitLFOs(*BusPtr.Get(), RefProxies.LFOs);
+				};
 
-			if (BusProxy && InSound)
-			{
-				BusProxy->OnInitSound(*InSound);
+				FBusHandle BusHandle = FBusHandle::Create(*BusPtr.Get(), RefProxies.Buses, OnCreate);
+				ManuallyActivatedBuses.Add(MoveTemp(BusHandle));
 			}
 		});
+
 	}
 
-	void FAudioModulationImpl::ActivateBusMix(const USoundControlBusMix& InBusMix, const ISoundModulatable* InSound)
+	void FAudioModulationImpl::ActivateBusMix(const USoundControlBusMix& InBusMix)
 	{
-		const FModulatorBusMixProxy NewMixProxy(InBusMix);
-		const bool bCanCreateNew = InBusMix.CanAutoActivate(InSound);
-
-		RunCommandOnAudioThread([this, NewMixProxy, InSound, bCanCreateNew]()
+		const TWeakObjectPtr<const USoundControlBusMix>BusMixPtr = &InBusMix;
+		RunCommandOnAudioThread([this, BusMixPtr]()
 		{
-			const FBusMixId MixId = static_cast<FBusMixId>(NewMixProxy.GetId());
-			FModulatorBusMixProxy* BusMixProxy = ActiveBusMixes.Find(MixId);
-			if (BusMixProxy)
-			{
-				// Enable in case mix is currently stopping but not yet stopped.
-				BusMixProxy->SetEnabled();
-			}
-			else if (bCanCreateNew)
-			{
-				BusMixProxy = &ActiveBusMixes.Add(MixId, NewMixProxy);
-			}
-
-			if (BusMixProxy && InSound)
+			if (BusMixPtr.IsValid())
 			{
-				BusMixProxy->OnInitSound(*InSound);
+				FBusMixHandle BusMixHandle = FBusMixHandle::Create(*BusMixPtr.Get(), RefProxies.BusMixes);
+				ManuallyActivatedBusMixes.Add(MoveTemp(BusMixHandle));
 			}
 		});
+
 	}
 
-	void FAudioModulationImpl::ActivateLFO(const USoundBusModulatorLFO& InLFO, const ISoundModulatable* InSound)
+	void FAudioModulationImpl::ActivateLFO(const USoundBusModulatorLFO& InLFO)
 	{
-		const FModulatorLFOProxy NewLFOProxy(InLFO);
-		const bool bCanCreateNew = InLFO.CanAutoActivate(InSound);
-
-		RunCommandOnAudioThread([this, NewLFOProxy, InSound, bCanCreateNew]()
+		const TWeakObjectPtr<const USoundBusModulatorLFO>LFOPtr = &InLFO;
+		RunCommandOnAudioThread([this, LFOPtr]()
 		{
-			const AudioModulation::FLFOId LFOId = NewLFOProxy.GetId();
-			FModulatorLFOProxy* LFOProxy = ActiveLFOs.Find(LFOId);
-			if (!LFOProxy && bCanCreateNew)
-			{
-				LFOProxy = &ActiveLFOs.Add(LFOId, NewLFOProxy);
-			}
-
-			if (LFOProxy && InSound)
+			if (LFOPtr.IsValid())
 			{
-				LFOProxy->OnInitSound(*InSound);
+				FLFOHandle LFOHandle = FLFOHandle::Create(*LFOPtr.Get(), RefProxies.LFOs);
+				ManuallyActivatedLFOs.Add(MoveTemp(LFOHandle));
 			}
 		});
 	}
@@ -479,11 +237,12 @@ namespace AudioModulation
 		{
 			if (InputProxy.bSampleAndHold)
 			{
-				if (!OutProxy.OutputProxy.bInitialized)
+				if (!OutProxy.OutputProxy.bInitialized && InputProxy.BusHandle.IsValid())
 				{
-					if (const FControlBusProxy* BusProxy = ActiveBuses.Find(InputProxy.BusId))
+					const FControlBusProxy& BusProxy = InputProxy.BusHandle.FindProxy();
+					if (!BusProxy.IsBypassed())
 					{
-						float ModStageValue = BusProxy->GetValue();
+						float ModStageValue = BusProxy.GetValue();
 						InputProxy.Transform.Apply(ModStageValue);
 						MixInModulationValue(OutProxy.OutputProxy.Operator, ModStageValue, OutSampleHold);
 					}
@@ -491,11 +250,15 @@ namespace AudioModulation
 			}
 			else
 			{
-				if (const FControlBusProxy* BusProxy = ActiveBuses.Find(InputProxy.BusId))
+				if (InputProxy.BusHandle.IsValid())
 				{
-					float ModStageValue = BusProxy->GetValue();
-					InputProxy.Transform.Apply(ModStageValue);
-					MixInModulationValue(OutProxy.OutputProxy.Operator, ModStageValue, OutValue);
+					const FControlBusProxy& BusProxy = InputProxy.BusHandle.FindProxy();
+					if (!BusProxy.IsBypassed())
+					{
+						float ModStageValue = BusProxy.GetValue();
+						InputProxy.Transform.Apply(ModStageValue);
+						MixInModulationValue(OutProxy.OutputProxy.Operator, ModStageValue, OutValue);
+					}
 				}
 			}
 		}
@@ -526,97 +289,54 @@ namespace AudioModulation
 		return !FMath::IsNearlyEqual(InitValue, OutValue);
 	}
 
-	float FAudioModulationImpl::CalculateInitialVolume(const USoundModulationPluginSourceSettingsBase& SettingsBase) const
+	float FAudioModulationImpl::CalculateInitialVolume(const USoundModulationPluginSourceSettingsBase& SettingsBase)
 	{
 		check(IsInAudioThread());
 
 		const USoundModulationSettings* Settings = CastChecked<USoundModulationSettings>(&SettingsBase);
-		FModulationPatchProxy VolumePatch(Settings->Volume);
+		FModulationPatchProxy VolumePatch(Settings->Volume, RefProxies);
 
 		return CalculateModulationValue(VolumePatch);
 	}
 
-	void FAudioModulationImpl::DeactivateBusMix(const FBusMixId InBusMixId, const ISoundModulatable* InSound)
+	void FAudioModulationImpl::DeactivateBus(const USoundControlBusBase& InBus)
 	{
-		RunCommandOnAudioThread([this, InBusMixId, InSound]()
+		const TWeakObjectPtr<const USoundControlBusBase>BusPtr = &InBus;
+		RunCommandOnAudioThread([this, BusPtr]()
 		{
-			check(IsInAudioThread());
-
-			if (FModulatorBusMixProxy* Mix = ActiveBusMixes.Find(InBusMixId))
+			if (BusPtr.IsValid())
 			{
-				if (!InSound)
-				{
-					if (!Mix->GetAutoActivate())
-					{
-						Mix->SetStopping();
-					}
-				}
-				else if (Mix->OnReleaseSound(*InSound) == 0)
-				{
-					if (Mix->GetAutoActivate())
-					{
-						Mix->SetStopping();
-					}
-				}
+				const USoundControlBusBase& Bus = *BusPtr.Get();
+				FBusHandle BusHandle = FBusHandle::Create(Bus, RefProxies.Buses);
+				ManuallyActivatedBuses.Remove(MoveTemp(BusHandle));
 			}
 		});
 	}
 
-	void FAudioModulationImpl::DeactivateBus(const FBusId InBusId, const ISoundModulatable* InSound)
+	void FAudioModulationImpl::DeactivateBusMix(const USoundControlBusMix& InBusMix)
 	{
-		RunCommandOnAudioThread([this, InBusId, InSound]()
+		const TWeakObjectPtr<const USoundControlBusMix> BusMixPtr = &InBusMix;
+		RunCommandOnAudioThread([this, BusMixPtr]()
 		{
-			check(IsInAudioThread());
-
-			if (FControlBusProxy* Bus = ActiveBuses.Find(InBusId))
+			if (BusMixPtr.IsValid())
 			{
-				// Only pass along to referenced LFOs if deactivating
-				// via notification of sound release
-				for (FLFOId LFOId : Bus->GetLFOIds())
-				{
-					DeactivateLFO(LFOId, InSound);
-				}
-
-				if (!InSound)
-				{
-					if (!Bus->GetAutoActivate())
-					{
-						ActiveBuses.Remove(InBusId);
-					}
-				}
-				else if (Bus->OnReleaseSound(*InSound) == 0)
-				{
-					if (Bus->GetAutoActivate())
-					{
-						ActiveBuses.Remove(InBusId);
-					}
-				}
+				const USoundControlBusMix& BusMix = *BusMixPtr.Get();
+				FBusMixHandle BusMixHandle = FBusMixHandle::Create(BusMix, RefProxies.BusMixes);
+				ManuallyActivatedBusMixes.Remove(MoveTemp(BusMixHandle));
 			}
 		});
 	}
 
-	void FAudioModulationImpl::DeactivateLFO(const FLFOId InLFOId, const ISoundModulatable* InSound)
+	void FAudioModulationImpl::DeactivateLFO(const USoundBusModulatorLFO& InLFO)
 	{
-		RunCommandOnAudioThread([this, InLFOId, InSound]()
+		const TWeakObjectPtr<const USoundBusModulatorLFO>LFOPtr = &InLFO;
+		RunCommandOnAudioThread([this, LFOPtr]()
 		{
-			check(IsInAudioThread());
-
-			if (FModulatorLFOProxy* LFO = ActiveLFOs.Find(InLFOId))
+			if (LFOPtr.IsValid())
 			{
-				if (!InSound)
-				{
-					if (!LFO->GetAutoActivate())
-					{
-						ActiveLFOs.Remove(InLFOId);
-					}
-				}
-				else if (LFO->OnReleaseSound(*InSound) == 0)
-				{
-					if (LFO->GetAutoActivate())
-					{
-						ActiveLFOs.Remove(InLFOId);
-					}
-				}
+				const USoundBusModulatorLFO& LFO = *LFOPtr.Get();
+				FLFOHandle LFOHandle = FLFOHandle::Create(LFO, RefProxies.LFOs);
+				ManuallyActivatedLFOs.Remove(MoveTemp(LFOHandle));
 			}
 		});
 	}
@@ -625,14 +345,7 @@ namespace AudioModulation
 	{
 		check(IsInAudioThread());
 
-		return ActiveBuses.Contains(InBusId);
-	}
-
-	bool FAudioModulationImpl::IsLFOActive(const FLFOId InLFOId) const
-	{
-		check(IsInAudioThread());
-
-		return ActiveLFOs.Contains(InLFOId);
+		return RefProxies.Buses.Contains(InBusId);
 	}
 
 	bool FAudioModulationImpl::ProcessControls(const uint32 InSourceId, FSoundModulationControls& OutControls)
@@ -696,24 +409,24 @@ namespace AudioModulation
 		check(IsInAudioThread());
 
 		// Update LFOs (prior to bus mixing to avoid single-frame latency)
-		for (TPair<FLFOId, FModulatorLFOProxy>& Pair : ActiveLFOs)
+		for (TPair<FLFOId, FModulatorLFOProxy>& Pair : RefProxies.LFOs)
 		{
 			Pair.Value.Update(Elapsed);
 		}
 
 		// Reset buses & refresh cached LFO
-		for (TPair<FBusId, FControlBusProxy>& Pair : ActiveBuses)
+		for (TPair<FBusId, FControlBusProxy>& Pair : RefProxies.Buses)
 		{
 			Pair.Value.Reset();
-			Pair.Value.MixLFO(ActiveLFOs);
+			Pair.Value.MixLFO();
 		}
 
 		// Update mix values and apply to prescribed buses.
 		// Track bus mixes ready to remove
 		TArray<FBusMixId> MixesToDeactivate;
-		for (TPair<FBusMixId, FModulatorBusMixProxy>& Pair : ActiveBusMixes)
+		for (TPair<FBusMixId, FModulatorBusMixProxy>& Pair : RefProxies.BusMixes)
 		{
-			Pair.Value.Update(Elapsed, ActiveBuses);
+			Pair.Value.Update(Elapsed, RefProxies.Buses);
 			if (Pair.Value.CanDestroy())
 			{
 				UE_LOG(LogAudioModulation, Log, TEXT("Audio modulation mix '%s' stopped."), *Pair.Value.GetName());
@@ -724,15 +437,15 @@ namespace AudioModulation
 		// Destroy mixes that have stopped
 		for (const FBusMixId& MixId : MixesToDeactivate)
 		{
-			ActiveBusMixes.Remove(MixId);
+			RefProxies.BusMixes.Remove(MixId);
 		}
 
-		SET_DWORD_STAT(STAT_AudioModulationBusCount, ActiveBuses.Num());
-		SET_DWORD_STAT(STAT_AudioModulationMixCount, ActiveBusMixes.Num());
-		SET_DWORD_STAT(STAT_AudioModulationLFOCount, ActiveLFOs.Num());
+		SET_DWORD_STAT(STAT_AudioModulationBusCount, RefProxies.Buses.Num());
+		SET_DWORD_STAT(STAT_AudioModulationMixCount, RefProxies.BusMixes.Num());
+		SET_DWORD_STAT(STAT_AudioModulationLFOCount, RefProxies.LFOs.Num());
 
 #if !UE_BUILD_SHIPPING
-		Debugger.UpdateDebugData(ActiveBuses, ActiveBusMixes, ActiveLFOs);
+		Debugger.UpdateDebugData(RefProxies);
 #endif // !UE_BUILD_SHIPPING
 	}
 
@@ -753,7 +466,7 @@ namespace AudioModulation
 		const FBusMixId MixId = static_cast<FBusMixId>(InMix.GetUniqueID());
 		RunCommandOnAudioThread([this, MixId, InChannels]()
 		{
-			if (FModulatorBusMixProxy* BusMixes = ActiveBusMixes.Find(MixId))
+			if (FModulatorBusMixProxy* BusMixes = RefProxies.BusMixes.Find(MixId))
 			{
 				BusMixes->SetMix(InChannels);
 			}
@@ -772,7 +485,7 @@ namespace AudioModulation
 
 		RunCommandOnAudioThread([this, ClassId, MixId, AddressFilter, InValue]()
 		{
-			if (FModulatorBusMixProxy* MixProxy = ActiveBusMixes.Find(MixId))
+			if (FModulatorBusMixProxy* MixProxy = RefProxies.BusMixes.Find(MixId))
 			{
 				MixProxy->SetMixByFilter(AddressFilter, ClassId, InValue);
 			}
@@ -783,39 +496,44 @@ namespace AudioModulation
 	{
 		if (const USoundBusModulatorLFO* LFO = Cast<USoundBusModulatorLFO>(&InModulator))
 		{
-			const FLFOId LFOId = static_cast<FLFOId>(InModulator.GetUniqueID());
-			const FModulatorLFOProxy UpdateProxy(*LFO);
-			RunCommandOnAudioThread([this, UpdateProxy, LFOId]()
+			const TWeakObjectPtr<const USoundBusModulatorLFO> LFOPtr = LFO;
+			RunCommandOnAudioThread([this, LFOPtr]()
 			{
-				if (FModulatorLFOProxy* LFOProxy = ActiveLFOs.Find(static_cast<FLFOId>(LFOId)))
+				if (LFOPtr.IsValid())
 				{
-					LFOProxy->OnUpdateProxy(UpdateProxy);
+					const USoundBusModulatorLFO& LFO = *LFOPtr.Get();
+					FLFOHandle LFOHandle = FLFOHandle::Create(LFO, RefProxies.LFOs);
+					LFOHandle.FindProxy() = LFO;
 				}
 			});
 		}
 
 		if (const USoundControlBusBase* Bus = Cast<USoundControlBusBase>(&InModulator))
 		{
-			const FBusId BusId = static_cast<FBusId>(InModulator.GetUniqueID());
-			const FControlBusProxy UpdateProxy(*Bus);
-			RunCommandOnAudioThread([this, UpdateProxy, BusId]()
+			const TWeakObjectPtr<const USoundControlBusBase> BusPtr = Bus;
+			RunCommandOnAudioThread([this, BusPtr]()
 			{
-				if (FControlBusProxy* BusProxy = ActiveBuses.Find(static_cast<FBusId>(BusId)))
+				if (BusPtr.IsValid())
 				{
-					BusProxy->OnUpdateProxy(UpdateProxy);
+					const USoundControlBusBase* Bus = BusPtr.Get();
+					FBusHandle BusHandle = FBusHandle::Create(*Bus, RefProxies.Buses);
+					FControlBusProxy& Proxy = BusHandle.FindProxy();
+					Proxy = *Bus;
+					Proxy.InitLFOs(*Bus, RefProxies.LFOs);
 				}
 			});
 		}
 
 		if (const USoundControlBusMix* Mix = Cast<USoundControlBusMix>(&InModulator))
 		{
-			const FBusMixId BusMixId = static_cast<FBusMixId>(InModulator.GetUniqueID());
-			const FModulatorBusMixProxy UpdateProxy(*Mix);
-			RunCommandOnAudioThread([this, UpdateProxy, BusMixId]()
+			const TWeakObjectPtr<const USoundControlBusMix> BusMixPtr = Mix;
+			RunCommandOnAudioThread([this, BusMixPtr]()
 			{
-				if (FModulatorBusMixProxy* BusMixProxy = ActiveBusMixes.Find(static_cast<FBusMixId>(BusMixId)))
+				if (BusMixPtr.IsValid())
 				{
-					BusMixProxy->OnUpdateProxy(UpdateProxy);
+					const USoundControlBusMix& BusMix = *BusMixPtr.Get();
+					FBusMixHandle BusMixHandle = FBusMixHandle::Create(BusMix, RefProxies.BusMixes);
+					BusMixHandle.FindProxy() = BusMix;
 				}
 			});
 		}
diff --git a/Engine/Plugins/Runtime/AudioModulation/Source/AudioModulation/Private/AudioModulationInternal.h b/Engine/Plugins/Runtime/AudioModulation/Source/AudioModulation/Private/AudioModulationInternal.h
--- a/Engine/Plugins/Runtime/AudioModulation/Source/AudioModulation/Private/AudioModulationInternal.h
+++ b/Engine/Plugins/Runtime/AudioModulation/Source/AudioModulation/Private/AudioModulationInternal.h
@@ -38,22 +38,21 @@ namespace AudioModulation
 		void OnReleaseSound(ISoundModulatable& InSound);
 		void OnReleaseSource(const uint32 InSourceId);
 
-		void ActivateBus(const USoundControlBusBase& InBus, const ISoundModulatable* Sound = nullptr);
-		void ActivateBusMix(const USoundControlBusMix& InBusMix, const ISoundModulatable* Sound = nullptr);
-		void ActivateLFO(const USoundBusModulatorLFO& InLFO, const ISoundModulatable* Sound = nullptr);
+		void ActivateBus(const USoundControlBusBase& InBus);
+		void ActivateBusMix(const USoundControlBusMix& InBusMix);
+		void ActivateLFO(const USoundBusModulatorLFO& InLFO);
 
-		float CalculateInitialVolume(const USoundModulationPluginSourceSettingsBase& InSettingsBase) const;
+		float CalculateInitialVolume(const USoundModulationPluginSourceSettingsBase& InSettingsBase);
 
 		/**
-		 * Deactivates respectively typed (i.e. BusMix, Bus, etc.) object proxy if either Sound is invalid (default behavior),
-		 * or if provided Sound, only deactivates if object is set to AutoActivate and sounds are no longer referencing provided object.
+		 * Deactivates respectively typed (i.e. BusMix, Bus, etc.) object proxy if no longer referenced.
+		 * If still referenced, will wait until references are finished before destroying.
 		 */
-		void DeactivateBusMix(const FBusMixId InBusMixId, const ISoundModulatable* Sound = nullptr);
-		void DeactivateBus(const FBusId InBusId, const ISoundModulatable* Sound = nullptr);
-		void DeactivateLFO(const FLFOId InLFOId, const ISoundModulatable* Sound = nullptr);
+		void DeactivateBus(const USoundControlBusBase& InBus);
+		void DeactivateBusMix(const USoundControlBusMix& InBusMix);
+		void DeactivateLFO(const USoundBusModulatorLFO& InLFO);
 
 		bool IsBusActive(const FBusId InBusId) const;
-		bool IsLFOActive(const FLFOId InLFOId) const;
 
 		bool ProcessControls(const uint32 InSourceId, FSoundModulationControls& OutControls);
 		void ProcessModulators(const float Elapsed);
@@ -69,15 +68,13 @@ namespace AudioModulation
 		/* Calculates modulation value, storing it in the provided float reference and returns if value changed */
 		bool CalculateModulationValue(FModulationPatchProxy& Proxy, float& OutValue) const;
 
-#if WITH_EDITOR
-		void EndPreviewSession();
-#endif // WITH_EDITOR
-
 		void RunCommandOnAudioThread(TFunction<void()> Cmd);
 
-		BusMixProxyMap ActiveBusMixes;
-		BusProxyMap    ActiveBuses;
-		LFOProxyMap    ActiveLFOs;
+		FReferencedProxies RefProxies;
+
+		TSet<FBusHandle>    ManuallyActivatedBuses;
+		TSet<FBusMixHandle> ManuallyActivatedBusMixes;
+		TSet<FLFOHandle>    ManuallyActivatedLFOs;
 
 		// Cache of SoundId to sound setting state at time of sound initialization.
 		TMap<uint32, FModulationSettingsProxy> SoundSettings;
@@ -94,12 +91,6 @@ namespace AudioModulation
 	private:
 		FAudioModulationDebugger Debugger;
 #endif // !UE_BUILD_SHIPPING
-
-#if WITH_EDITOR
-	private:
-		ISoundModulatable* PreviewSound;
-		FModulationSettingsProxy PreviewSettings;
-#endif // WITH_EDITOR
 	};
 } // namespace AudioModulation
 
@@ -128,15 +119,15 @@ namespace AudioModulation
 		bool OnToggleStat(FCommonViewportClient* ViewportClient, const TCHAR* Stream) { return false; }
 #endif // !UE_BUILD_SHIPPING
 
-		void ActivateBus(const USoundControlBusBase& InBus, const ISoundModulatable* Sound = nullptr) { }
-		void ActivateBusMix(const USoundControlBusMix& InBusMix, const ISoundModulatable* Sound = nullptr) { }
-		void ActivateLFO(const USoundBusModulatorLFO& InLFO, const ISoundModulatable* Sound = nullptr) { }
+		void ActivateBus(const USoundControlBusBase& InBus) { }
+		void ActivateBusMix(const USoundControlBusMix& InBusMix) { }
+		void ActivateLFO(const USoundBusModulatorLFO& InLFO) { }
 
-		float CalculateInitialVolume(const USoundModulationPluginSourceSettingsBase& InSettingsBase) const { return 1.0f; }
+		float CalculateInitialVolume(const USoundModulationPluginSourceSettingsBase& InSettingsBase) { return 1.0f; }
 
-		void DeactivateBusMix(const FBusMixId BusMixId, const ISoundModulatable* Sound = nullptr) { }
-		void DeactivateBus(const FBusId BusId, const ISoundModulatable* Sound = nullptr) { }
-		void DeactivateLFO(const FLFOId LFOId, const ISoundModulatable* Sound = nullptr) { }
+		void DeactivateBus(const USoundControlBusBase& InBus) { }
+		void DeactivateBusMix(const USoundControlBusBase& InBusMix) { }
+		void DeactivateLFO(const USoundBusModulatorLFO& InLFO) { }
 
 		bool ProcessControls(const uint32 InSourceId, FSoundModulationControls& OutControls) { return false; }
 		void ProcessModulators(const float Elapsed) { }
@@ -144,12 +135,6 @@ namespace AudioModulation
 		void UpdateMix(const USoundControlBusMix& InMix, const TArray<FSoundControlBusMixChannel>& InChannels) { }
 		void UpdateMixByFilter(const USoundControlBusMix& InMix, const FString& InAddressFilter, const TSubclassOf<USoundControlBusBase>& InClassFilter, const FSoundModulationValue& InValue) { }
 		void UpdateModulator(const USoundModulatorBase& InModulator) { }
-
-		void SetBusDefault(const USoundControlBusBase& InBus, const float Value) { }
-		void SetBusMin(const USoundControlBusBase& InBus, const float Value) { }
-		void SetBusMax(const USoundControlBusBase& InBus, const float Value) { }
-		void SetBusMixChannel(const USoundControlBusMix& InBusMix, const USoundControlBusBase& InBus, const float TargetValue) { }
-		void SetLFOFrequency(const USoundBusModulatorLFO& InLFO, const float Freq) { }
 	};
 } // namespace AudioModulation
 #endif // WITH_AUDIOMODULATION
diff --git a/Engine/Plugins/Runtime/AudioModulation/Source/AudioModulation/Private/AudioModulationStatics.cpp b/Engine/Plugins/Runtime/AudioModulation/Source/AudioModulation/Private/AudioModulationStatics.cpp
--- a/Engine/Plugins/Runtime/AudioModulation/Source/AudioModulation/Private/AudioModulationStatics.cpp
+++ b/Engine/Plugins/Runtime/AudioModulation/Source/AudioModulation/Private/AudioModulationStatics.cpp
@@ -216,16 +216,13 @@ USoundControlBusMix* UAudioModulationStatics::CreateBusMix(const UObject* WorldC
 
 void UAudioModulationStatics::DeactivateBus(const UObject* WorldContextObject, USoundControlBusBase* Bus)
 {
-	if (!Bus)
+	if (Bus)
 	{
-		return;
-	}
-
-	UWorld* World = GetAudioWorld(WorldContextObject);
-	if (AudioModulation::FAudioModulationImpl* ModulationImpl = GetModulationImpl(World))
-	{
-		auto BusId = static_cast<const AudioModulation::FBusId>(Bus->GetUniqueID());
-		ModulationImpl->DeactivateBus(BusId);
+		UWorld* World = GetAudioWorld(WorldContextObject);
+		if (AudioModulation::FAudioModulationImpl* ModulationImpl = GetModulationImpl(World))
+		{
+			ModulationImpl->DeactivateBus(*Bus);
+		}
 	}
 }
 
@@ -236,8 +233,7 @@ void UAudioModulationStatics::DeactivateBusMix(const UObject* WorldContextObject
 		UWorld* World = GetAudioWorld(WorldContextObject);
 		if (AudioModulation::FAudioModulationImpl* ModulationImpl = GetModulationImpl(World))
 		{
-			auto BusMixId = static_cast<const AudioModulation::FBusMixId>(BusMix->GetUniqueID());
-			ModulationImpl->DeactivateBusMix(BusMixId);
+			ModulationImpl->DeactivateBusMix(*BusMix);
 		}
 	}
 }
@@ -249,8 +245,7 @@ void UAudioModulationStatics::DeactivateBusModulator(const UObject* WorldContext
 	{
 		if (USoundBusModulatorLFO* LFO = Cast<USoundBusModulatorLFO>(Modulator))
 		{
-			auto LFOId = static_cast<const AudioModulation::FLFOId>(Modulator->GetUniqueID());
-			ModulationImpl->DeactivateLFO(LFOId);
+			ModulationImpl->DeactivateLFO(*LFO);
 		}
 	}
 }
diff --git a/Engine/Plugins/Runtime/AudioModulation/Source/AudioModulation/Private/SoundControlBus.cpp b/Engine/Plugins/Runtime/AudioModulation/Source/AudioModulation/Private/SoundControlBus.cpp
--- a/Engine/Plugins/Runtime/AudioModulation/Source/AudioModulation/Private/SoundControlBus.cpp
+++ b/Engine/Plugins/Runtime/AudioModulation/Source/AudioModulation/Private/SoundControlBus.cpp
@@ -102,65 +102,43 @@ void USoundControlBusBase::BeginDestroy()
 {
 	Super::BeginDestroy();
 
-	UWorld* World = GetWorld();
-	if (!World)
+	if (UWorld* World = GetWorld())
 	{
-		return;
-	}
-
-	FAudioDevice* AudioDevice = World->GetAudioDevice();
-	if (!AudioDevice)
-	{
-		return;
-	}
-
-	check(AudioDevice->IsModulationPluginEnabled());
-	if (IAudioModulation* ModulationInterface = AudioDevice->ModulationInterface.Get())
-	{
-		auto ModulationImpl = static_cast<AudioModulation::FAudioModulation*>(ModulationInterface)->GetImpl();
-		check(ModulationImpl);
-
-		auto BusId = static_cast<const AudioModulation::FBusId>(GetUniqueID());
-		ModulationImpl->DeactivateBus(BusId);
+		if (FAudioDevice* AudioDevice = World->GetAudioDevice())
+		{
+			check(AudioDevice->IsModulationPluginEnabled());
+			if (IAudioModulation* ModulationInterface = AudioDevice->ModulationInterface.Get())
+			{
+				auto ModulationImpl = static_cast<AudioModulation::FAudioModulation*>(ModulationInterface)->GetImpl();
+				check(ModulationImpl);
+				ModulationImpl->DeactivateBus(*this);
+			}
+		}
 	}
 }
 
 namespace AudioModulation
 {
 	FControlBusProxy::FControlBusProxy()
-		: TModulatorProxyRefBase<FBusId>()
-		, DefaultValue(0.0f)
+		: DefaultValue(0.0f)
 		, LFOValue(1.0f)
 		, MixValue(NAN)
+		, bBypass(false)
 		, Operator(ESoundModulatorOperator::Multiply)
 		, Range(0.0f, 1.0f)
 	{
 	}
 
-	FControlBusProxy::FControlBusProxy(const USoundControlBusBase& Bus)
-		: TModulatorProxyRefBase<FBusId>(Bus.GetName(), Bus.GetUniqueID(), Bus.bAutoActivate)
-		, DefaultValue(Bus.DefaultValue)
-		, LFOValue(1.0f)
-		, MixValue(NAN)
-		, Operator(Bus.GetOperator())
-		, Range(Bus.Min, Bus.Max)
+	FControlBusProxy::FControlBusProxy(const USoundControlBusBase& InBus)
+		: TModulatorProxyRefType(InBus.GetName(), InBus.GetUniqueID())
 	{
-		if (Bus.Min > Bus.Max)
-		{
-			Range.X = Bus.Max;
-			Range.Y = Bus.Min;
-		}
-
-		DefaultValue = FMath::Clamp(DefaultValue, Range.X, Range.Y);
+		Init(InBus);
+	}
 
-		for (const USoundBusModulatorBase* Modulator: Bus.Modulators)
-		{
-			if (const USoundBusModulatorLFO* LFO = Cast<USoundBusModulatorLFO>(Modulator))
-			{
-				const FLFOId RefLFOId = static_cast<FLFOId>(LFO->GetUniqueID());
-				LFOIds.Add(RefLFOId);
-			}
-		}
+	FControlBusProxy& FControlBusProxy::operator =(const USoundControlBusBase& InBus)
+	{
+		Init(InBus);
+		return *this;
 	}
 
 	float FControlBusProxy::GetDefaultValue() const
@@ -168,9 +146,9 @@ namespace AudioModulation
 		return DefaultValue;
 	}
 
-	const TArray<FLFOId>& FControlBusProxy::GetLFOIds() const
+	const TArray<FLFOHandle>& FControlBusProxy::GetLFOHandles() const
 	{
-		return LFOIds;
+		return LFOHandles;
 	}
 
 	float FControlBusProxy::GetLFOValue() const
@@ -194,6 +172,41 @@ namespace AudioModulation
 		return FMath::Clamp(DefaultMixed * LFOValue, Range.X, Range.Y);
 	}
 
+	void FControlBusProxy::Init(const USoundControlBusBase& InBus)
+	{
+		DefaultValue = InBus.DefaultValue;
+		LFOValue = 1.0f;
+		MixValue = NAN;
+		Operator = InBus.GetOperator();
+		Range = FVector2D(InBus.Min, InBus.Max);
+		if (InBus.Min > InBus.Max)
+		{
+			Range.X = InBus.Max;
+			Range.Y = InBus.Min;
+		}
+
+		DefaultValue = FMath::Clamp(DefaultValue, Range.X, Range.Y);
+		bBypass = InBus.bBypass;
+	}
+
+	void FControlBusProxy::InitLFOs(const USoundControlBusBase& InBus, FLFOProxyMap& OutActiveLFOs)
+	{
+		LFOHandles.Reset();
+
+		for (const USoundBusModulatorBase* Modulator : InBus.Modulators)
+		{
+			if (const USoundBusModulatorLFO* LFO = Cast<USoundBusModulatorLFO>(Modulator))
+			{
+				LFOHandles.Add(FLFOHandle::Create(*LFO, OutActiveLFOs));
+			}
+		}
+	}
+
+	bool FControlBusProxy::IsBypassed() const
+	{
+		return bBypass;
+	}
+
 	float FControlBusProxy::Mix(float ValueA) const
 	{
 		// If mix value is NaN, it is uninitialized (effectively, the parent bus is inactive)
@@ -243,24 +256,21 @@ namespace AudioModulation
 		MixValue = Mix(InValue);
 	}
 
-	void FControlBusProxy::MixLFO(LFOProxyMap& LFOMap)
+	void FControlBusProxy::MixLFO()
 	{
-		for (const FLFOId& LFOId : LFOIds)
+		for (const FLFOHandle& Handle: LFOHandles)
 		{
-			if (FModulatorLFOProxy* LFOProxy = LFOMap.Find(LFOId))
+			if (Handle.IsValid())
 			{
-				LFOValue *= LFOProxy->GetValue();
+				const FModulatorLFOProxy& LFOProxy = Handle.FindProxy();
+				if (!LFOProxy.IsBypassed())
+				{
+					LFOValue *= LFOProxy.GetValue();
+				}
 			}
 		}
 	}
 
-	void FControlBusProxy::OnUpdateProxy(const FControlBusProxy& InBusProxy)
-	{
-		DefaultValue = FMath::Clamp(InBusProxy.DefaultValue, InBusProxy.Range.GetMin(), InBusProxy.Range.GetMax());
-		Operator = InBusProxy.Operator;
-		Range = InBusProxy.Range;
-	}
-
 	void FControlBusProxy::Reset()
 	{
 		LFOValue = 1.0f;
diff --git a/Engine/Plugins/Runtime/AudioModulation/Source/AudioModulation/Private/SoundControlBusMix.cpp b/Engine/Plugins/Runtime/AudioModulation/Source/AudioModulation/Private/SoundControlBusMix.cpp
--- a/Engine/Plugins/Runtime/AudioModulation/Source/AudioModulation/Private/SoundControlBusMix.cpp
+++ b/Engine/Plugins/Runtime/AudioModulation/Source/AudioModulation/Private/SoundControlBusMix.cpp
@@ -29,26 +29,18 @@ void USoundControlBusMix::BeginDestroy()
 {
 	Super::BeginDestroy();
 
-	UWorld* World = GetWorld();
-	if (!World)
+	if (UWorld* World = GetWorld())
 	{
-		return;
-	}
-
-	FAudioDevice* AudioDevice = World->GetAudioDevice();
-	if (!AudioDevice)
-	{
-		return;
-	}
-
-	check(AudioDevice->IsModulationPluginEnabled());
-	if (IAudioModulation* ModulationInterface = AudioDevice->ModulationInterface.Get())
-	{
-		auto ModulationImpl = static_cast<AudioModulation::FAudioModulation*>(ModulationInterface)->GetImpl();
-		check(ModulationImpl);
-
-		auto BusMixId = static_cast<const AudioModulation::FBusMixId>(GetUniqueID());
-		ModulationImpl->DeactivateBusMix(BusMixId);
+		if (FAudioDevice* AudioDevice = World->GetAudioDevice())
+		{
+			check(AudioDevice->IsModulationPluginEnabled());
+			if (IAudioModulation* ModulationInterface = AudioDevice->ModulationInterface.Get())
+			{
+				auto ModulationImpl = static_cast<AudioModulation::FAudioModulation*>(ModulationInterface)->GetImpl();
+				check(ModulationImpl);
+				ModulationImpl->DeactivateBusMix(*this);
+			}
+		}
 	}
 }
 
@@ -62,11 +54,30 @@ namespace AudioModulation
 	{
 	}
 
-	FModulatorBusMixProxy::FModulatorBusMixProxy(const USoundControlBusMix& Mix)
-		: TModulatorProxyRefBase<FBusMixId>(Mix.GetName(), Mix.GetUniqueID(), Mix.bAutoActivate)
-		, Status(BusMixStatus::Enabled)
+	FModulatorBusMixProxy::FModulatorBusMixProxy(const USoundControlBusMix& InBusMix)
+		: TModulatorProxyRefType(InBusMix.GetName(), InBusMix.GetUniqueID())
+	{
+		Init(InBusMix);
+	}
+
+	FModulatorBusMixProxy& FModulatorBusMixProxy::operator =(const USoundControlBusMix& InBusMix)
+	{
+		Init(InBusMix);
+
+		return *this;
+	}
+
+	bool FModulatorBusMixProxy::CanDestroy() const
+	{
+		return Status == BusMixStatus::Stopped;
+	}
+
+	void FModulatorBusMixProxy::Init(const USoundControlBusMix& InBusMix)
 	{
-		for (const FSoundControlBusMixChannel& Channel : Mix.Channels)
+		Channels.Reset();
+
+		Status = BusMixStatus::Enabled;
+		for (const FSoundControlBusMixChannel& Channel : InBusMix.Channels)
 		{
 			if (Channel.Bus)
 			{
@@ -76,7 +87,7 @@ namespace AudioModulation
 				{
 					UE_LOG(LogAudioModulation, Warning,
 						TEXT("USoundControlBusMix '%s' already contains bus '%s'. Only one representative channel for this bus added."),
-						*Mix.GetFullName(), *Channel.Bus->GetFullName());
+						*InBusMix.GetFullName(), *Channel.Bus->GetFullName());
 				}
 #endif // UE_BUILD_SHIPPING
 				Channels.Emplace(BusId, FModulatorBusMixChannelProxy(Channel));
@@ -85,27 +96,12 @@ namespace AudioModulation
 			{
 				UE_LOG(LogAudioModulation, Warning,
 					TEXT("USoundControlBusMix '%s' has channel with no bus specified. "
-					"Mix activated but channel ignored."),
-					*Mix.GetFullName());
+						"Mix activated but channel ignored."),
+					*InBusMix.GetFullName());
 			}
 		}
 	}
 
-	bool FModulatorBusMixProxy::CanDestroy() const
-	{
-		if (Status != BusMixStatus::Stopped)
-		{
-			return false;
-		}
-
-		return TModulatorProxyRefBase<FBusMixId>::CanDestroy();
-	}
-
-	void FModulatorBusMixProxy::OnUpdateProxy(const FModulatorBusMixProxy& InBusMixProxy)
-	{
-		Channels = InBusMixProxy.Channels;
-	}
-
 	void FModulatorBusMixProxy::SetMix(const TArray<FSoundControlBusMixChannel>& InChannels)
 	{
 		for (const FSoundControlBusMixChannel& NewChannel : InChannels)
@@ -167,7 +163,7 @@ namespace AudioModulation
 		}
 	}
 
-	void FModulatorBusMixProxy::Update(float Elapsed, BusProxyMap& ProxyMap)
+	void FModulatorBusMixProxy::Update(float Elapsed, FBusProxyMap& ProxyMap)
 	{
 		bool bRequestStop = true;
 		for (TPair<FBusId, FModulatorBusMixChannelProxy>& Channel : Channels)
diff --git a/Engine/Plugins/Runtime/AudioModulation/Source/AudioModulation/Private/SoundModulationPatch.cpp b/Engine/Plugins/Runtime/AudioModulation/Source/AudioModulation/Private/SoundModulationPatch.cpp
--- a/Engine/Plugins/Runtime/AudioModulation/Source/AudioModulation/Private/SoundModulationPatch.cpp
+++ b/Engine/Plugins/Runtime/AudioModulation/Source/AudioModulation/Private/SoundModulationPatch.cpp
@@ -133,44 +133,59 @@ void FSoundHPFModulationPatch::Clamp()
 }
 #endif // WITH_EDITOR
 
-void FSoundVolumeModulationPatch::GenerateProxies(TArray<AudioModulation::FModulationInputProxy>& InputProxies) const
+TArray<const FSoundModulationInputBase*> FSoundVolumeModulationPatch::GetInputs() const
 {
+	TArray<const FSoundModulationInputBase*> OutInputs;
 	for (const FSoundVolumeModulationInput& Input : Inputs)
 	{
-		InputProxies.Emplace_GetRef(Input);
+		OutInputs.Add(static_cast<const FSoundModulationInputBase*>(&Input));
 	}
+
+	return OutInputs;
 }
 
-void FSoundPitchModulationPatch::GenerateProxies(TArray<AudioModulation::FModulationInputProxy>& InputProxies) const
+TArray<const FSoundModulationInputBase*> FSoundPitchModulationPatch::GetInputs() const
 {
+	TArray<const FSoundModulationInputBase*> OutInputs;
 	for (const FSoundPitchModulationInput& Input : Inputs)
 	{
-		InputProxies.Emplace_GetRef(Input);
+		OutInputs.Add(static_cast<const FSoundModulationInputBase*>(&Input));
 	}
+
+	return OutInputs;
 }
 
-void FSoundHPFModulationPatch::GenerateProxies(TArray<AudioModulation::FModulationInputProxy>& InputProxies) const
+TArray<const FSoundModulationInputBase*> FSoundLPFModulationPatch::GetInputs() const
 {
-	for (const FSoundHPFModulationInput& Input : Inputs)
+	TArray<const FSoundModulationInputBase*> OutInputs;
+	for (const FSoundLPFModulationInput& Input : Inputs)
 	{
-		InputProxies.Emplace_GetRef(Input);
+		OutInputs.Add(static_cast<const FSoundModulationInputBase*>(&Input));
 	}
+
+	return OutInputs;
 }
 
-void FSoundLPFModulationPatch::GenerateProxies(TArray<AudioModulation::FModulationInputProxy>& InputProxies) const
+TArray<const FSoundModulationInputBase*> FSoundHPFModulationPatch::GetInputs() const
 {
-	for (const FSoundLPFModulationInput& Input : Inputs)
+	TArray<const FSoundModulationInputBase*> OutInputs;
+	for (const FSoundHPFModulationInput& Input : Inputs)
 	{
-		InputProxies.Emplace_GetRef(Input);
+		OutInputs.Add(static_cast<const FSoundModulationInputBase*>(&Input));
 	}
+
+	return OutInputs;
 }
 
-void FSoundControlModulationPatch::GenerateProxies(TArray<AudioModulation::FModulationInputProxy>& InputProxies) const
+TArray<const FSoundModulationInputBase*> FSoundControlModulationPatch::GetInputs() const
 {
+	TArray<const FSoundModulationInputBase*> OutInputs;
 	for (const FSoundControlModulationInput& Input : Inputs)
 	{
-		InputProxies.Emplace_GetRef(Input);
+		OutInputs.Add(static_cast<const FSoundModulationInputBase*>(&Input));
 	}
+
+	return OutInputs;
 }
 
 FSoundModulationOutput::FSoundModulationOutput()
@@ -222,19 +237,22 @@ FSoundModulationPatchBase::FSoundModulationPatchBase()
 namespace AudioModulation
 {
 	FModulationInputProxy::FModulationInputProxy()
-		: BusId(InvalidBusId)
-		, bSampleAndHold(0)
+		: bSampleAndHold(0)
 	{
 	}
 
-	FModulationInputProxy::FModulationInputProxy(const FSoundModulationInputBase& Input)
-		: BusId(InvalidBusId)
-		, Transform(Input.Transform)
+	FModulationInputProxy::FModulationInputProxy(const FSoundModulationInputBase& Input, FReferencedProxies& OutRefProxies)
+		: Transform(Input.Transform)
 		, bSampleAndHold(Input.bSampleAndHold)
 	{
 		if (const USoundControlBusBase* Bus = Input.GetBus())
 		{
-			BusId = static_cast<AudioModulation::FBusId>(Bus->GetUniqueID());
+			FReferencedProxies* RefProxies = &OutRefProxies;
+			auto OnCreate = [RefProxies, Bus](FControlBusProxy& NewProxy)
+			{
+				NewProxy.InitLFOs(*Bus, RefProxies->LFOs);
+			};
+			BusHandle = FBusHandle::Create(*Bus, OutRefProxies.Buses, OnCreate);
 		}
 	}
 
@@ -259,12 +277,16 @@ namespace AudioModulation
 	{
 	}
 
-	FModulationPatchProxy::FModulationPatchProxy(const FSoundModulationPatchBase& Patch)
+	FModulationPatchProxy::FModulationPatchProxy(const FSoundModulationPatchBase& Patch, FReferencedProxies& OutRefProxies)
 		: DefaultInputValue(Patch.DefaultInputValue)
 		, bBypass(Patch.bBypass)
 		, OutputProxy(*Patch.GetOutput())
 	{
-		Patch.GenerateProxies(InputProxies);
+		TArray<const FSoundModulationInputBase*> Inputs = Patch.GetInputs();
+		for (const FSoundModulationInputBase* Input : Inputs)
+		{
+			InputProxies.Emplace_GetRef(*Input, OutRefProxies);
+		}
 	}
 
 	FModulationSettingsProxy::FModulationSettingsProxy()
@@ -284,23 +306,23 @@ namespace AudioModulation
 		Highpass.OutputProxy.Operator = ESoundModulatorOperator::Max;
 	}
 
-	FModulationSettingsProxy::FModulationSettingsProxy(const USoundModulationSettings& Settings)
+	FModulationSettingsProxy::FModulationSettingsProxy(const USoundModulationSettings& Settings, FReferencedProxies& OutRefProxies)
 		: TModulatorProxyBase<uint32>(Settings.GetName(), Settings.GetUniqueID())
-		, Volume(Settings.Volume)
-		, Pitch(Settings.Pitch)
-		, Lowpass(Settings.Lowpass)
-		, Highpass(Settings.Highpass)
+		, Volume(Settings.Volume, OutRefProxies)
+		, Pitch(Settings.Pitch, OutRefProxies)
+		, Lowpass(Settings.Lowpass, OutRefProxies)
+		, Highpass(Settings.Highpass, OutRefProxies)
 	{
 		for (const FSoundControlModulationPatch& Patch : Settings.Controls)
 		{
-			Controls.Add(Patch.Control, Patch);
+			Controls.Add(Patch.Control, FModulationPatchProxy(Patch, OutRefProxies));
 		}
 
 		for (const USoundControlBusMix* Mix : Settings.Mixes)
 		{
 			if (Mix)
 			{
-				Mixes.Add(static_cast<FBusMixId>(Mix->GetUniqueID()));
+				Mixes.Add(FBusMixHandle::Create(*Mix, OutRefProxies.BusMixes));
 			}
 		}
 	}
diff --git a/Engine/Plugins/Runtime/AudioModulation/Source/AudioModulation/Private/SoundModulationProxy.h b/Engine/Plugins/Runtime/AudioModulation/Source/AudioModulation/Private/SoundModulationProxy.h
--- a/Engine/Plugins/Runtime/AudioModulation/Source/AudioModulation/Private/SoundModulationProxy.h
+++ b/Engine/Plugins/Runtime/AudioModulation/Source/AudioModulation/Private/SoundModulationProxy.h
@@ -5,15 +5,21 @@
 #include "IAudioExtensionPlugin.h"
 #include "UObject/ObjectMacros.h"
 #include "UObject/Object.h"
+#include "UObject/WeakObjectPtrTemplates.h"
 
+#include "AudioModulationLogging.h"
 #include "SoundControlBus.h"
 #include "SoundControlBusMix.h"
 #include "SoundModulationPatch.h"
 #include "SoundModulationTransform.h"
+#include "SoundModulatorLFO.h"
 
 
 namespace AudioModulation
 {
+	// Forward Declarations
+	struct FReferencedProxies;
+
 	// Modulator Ids
 	using FBusMixId = uint32;
 	extern const FBusMixId InvalidBusMixId;
@@ -22,8 +28,182 @@ namespace AudioModulation
 	extern const FBusId InvalidBusId;
 
 	using FLFOId = uint32;
-	extern const FLFOId InvalidLFOId;
+	extern const FBusId InvalidLFOId;
+
+	/*
+	 * Handle for all ref-counted proxy types, to be used only on the audio thread (not thread safe). 
+	 */
+	template<typename IdType, typename ProxyType, typename ProxyUObjType>
+	class TProxyHandle
+	{
+		IdType Id;
+		TMap<IdType, ProxyType>* ProxyMap;
+
+	public:
+		FORCEINLINE bool IsValid() const
+		{
+			return static_cast<uint32>(Id) != INDEX_NONE;
+		}
 
+		FORCEINLINE IdType GetId() const
+		{
+			return Id;
+		}
+
+		FORCEINLINE ProxyType& FindProxy()
+		{
+			check(IsValid());
+			check(ProxyMap);
+
+			return ProxyMap->FindChecked(Id);
+		}
+
+		FORCEINLINE const ProxyType& FindProxy() const
+		{
+			check(IsValid());
+			check(ProxyMap);
+
+			return ProxyMap->FindChecked(Id);
+		}
+
+		friend FORCEINLINE uint32 GetTypeHash(const TProxyHandle<IdType, ProxyType, ProxyUObjType>& InHandle)
+		{
+			return static_cast<uint32>(InHandle.Id);
+		}
+
+		TProxyHandle()
+			: Id(static_cast<IdType>(INDEX_NONE))
+			, ProxyMap(nullptr)
+		{
+		}
+
+		TProxyHandle(IdType InId, TMap<IdType, ProxyType>& InProxyMap)
+			: Id(InId)
+			, ProxyMap(&InProxyMap)
+		{
+		}
+
+		TProxyHandle(const TProxyHandle& InHandle)
+		{
+			Id = InHandle.Id;
+			ProxyMap = InHandle.ProxyMap;
+
+			if (IsValid())
+			{
+				ProxyType& Proxy = ProxyMap->FindChecked(Id);
+				Proxy.IncRef();
+			}
+		}
+
+		TProxyHandle(TProxyHandle&& InHandle)
+		{
+			Id = InHandle.Id;
+			ProxyMap = InHandle.ProxyMap;
+
+			InHandle.Id = static_cast<IdType>(INDEX_NONE);
+			InHandle.ProxyMap = nullptr;
+		}
+
+		~TProxyHandle()
+		{
+			if (ProxyMap)
+			{
+				if (ProxyType* Proxy = ProxyMap->Find(Id))
+				{
+					const uint32 RefCount = Proxy->DecRef();
+					if (RefCount == 0)
+					{
+						UE_LOG(LogAudioModulation, Verbose, TEXT("Proxy released: Id '%u'."), Id);
+						ProxyMap->Remove(Id);
+					}
+				}
+			}
+		}
+
+		/*
+		 * Creates a handle to a proxy modulation object tracked in the provided InProxyMap
+		 */
+		static TProxyHandle<IdType, ProxyType, ProxyUObjType> Create(const ProxyUObjType& InObject, TMap<IdType, ProxyType>& InProxyMap)
+		{
+			const IdType ObjectId = static_cast<IdType>(InObject.GetUniqueID());
+			TProxyHandle<IdType, ProxyType, ProxyUObjType> NewHandle(ObjectId, InProxyMap);
+
+			ProxyType* Proxy = InProxyMap.Find(NewHandle.Id);
+			if (!Proxy)
+			{
+				UE_LOG(LogAudioModulation, Verbose, TEXT("Proxy created: Id '%u' for object '%s'."), NewHandle.Id, *InObject.GetName());
+				Proxy = &InProxyMap.Add(NewHandle.Id, ProxyType(InObject));
+			}
+			Proxy->IncRef();
+
+			return NewHandle;
+		}
+
+		static TProxyHandle<IdType, ProxyType, ProxyUObjType> Create(const ProxyUObjType& InObject, TMap<IdType, ProxyType>& InProxyMap, TFunction<void(ProxyType&)> OnCreateProxy)
+		{
+			const IdType ObjectId = static_cast<IdType>(InObject.GetUniqueID());
+			TProxyHandle<IdType, ProxyType, ProxyUObjType> NewHandle(ObjectId, InProxyMap);
+
+			ProxyType* Proxy = InProxyMap.Find(NewHandle.Id);
+			if (!Proxy)
+			{
+				UE_LOG(LogAudioModulation, Verbose, TEXT("Proxy created: Id '%u' for object '%s'."), NewHandle.Id, *InObject.GetName());
+				Proxy = &InProxyMap.Add(NewHandle.Id, ProxyType(InObject));
+				OnCreateProxy(*Proxy);
+			}
+			Proxy->IncRef();
+
+			return NewHandle;
+		}
+
+		FORCEINLINE bool operator ==(const TProxyHandle<IdType, ProxyType, ProxyUObjType>& InHandle) const
+		{
+			return InHandle.Id == Id;
+		}
+
+		FORCEINLINE TProxyHandle<IdType, ProxyType, ProxyUObjType>& operator =(const TProxyHandle<IdType, ProxyType, ProxyUObjType>& InHandle)
+		{
+			if (InHandle.Id == Id)
+			{
+				check(ProxyMap == InHandle.ProxyMap);
+				return *this;
+			}
+
+			if (IsValid())
+			{
+				ProxyType& Proxy = ProxyMap->FindChecked(Id);
+				Proxy.DecRef();
+			}
+
+			Id = InHandle.Id;
+			ProxyMap = InHandle.ProxyMap;
+
+			if (IsValid())
+			{
+				ProxyType& Proxy = ProxyMap->FindChecked(Id);
+				Proxy.IncRef();
+			}
+
+			return *this;
+		}
+
+		FORCEINLINE TProxyHandle<IdType, ProxyType, ProxyUObjType>& operator =(TProxyHandle<IdType, ProxyType, ProxyUObjType>&& InHandle)
+		{
+			if (IsValid())
+			{
+				ProxyType& Proxy = ProxyMap->FindChecked(Id);
+				Proxy.DecRef();
+			}
+
+			Id = InHandle.Id;
+			ProxyMap = InHandle.ProxyMap;
+
+			InHandle.Id = static_cast<IdType>(INDEX_NONE);
+			InHandle.ProxyMap = nullptr;
+
+			return *this;
+		}
+	};
 
 	// Modulator Proxy Templates
 	template <typename IdType>
@@ -38,7 +218,7 @@ namespace AudioModulation
 
 	public:
 		TModulatorProxyBase()
-			: Id(static_cast<IdType>(0))
+			: Id(static_cast<IdType>(INDEX_NONE))
 		{
 		}
 
@@ -67,116 +247,135 @@ namespace AudioModulation
 		}
 	};
 
-	template <typename IdType>
-	class TModulatorProxyRefBase : public TModulatorProxyBase<IdType>
+	struct FModulatorBusMixChannelProxy : public TModulatorProxyBase<FBusId>
+	{
+		FModulatorBusMixChannelProxy(const FSoundControlBusMixChannel& Channel);
+		FString Address;
+		uint32 ClassId;
+		FSoundModulationValue Value;
+	};
+
+	class FModulatorLFOProxy;
+
+	using FLFOProxyMap = TMap<FLFOId, FModulatorLFOProxy>;
+	using FLFOHandle = TProxyHandle<FLFOId, FModulatorLFOProxy, USoundBusModulatorLFO>;
+
+	template<typename IdType, typename ProxyType, typename ProxyUObjType>
+	class TModulatorProxyRefType : public TModulatorProxyBase<IdType>
 	{
+	protected:
+		uint32 RefCount;
+
 	public:
-		TModulatorProxyRefBase()
-			: bAutoActivate(false)
+		TModulatorProxyRefType()
+			: TModulatorProxyBase<IdType>()
+			, RefCount(0)
 		{
 		}
 
-		TModulatorProxyRefBase(const FString& Name, const IdType InId, const bool bInAutoActivate)
-			: TModulatorProxyBase<IdType>(Name, InId)
-			, bAutoActivate(bInAutoActivate)
+		TModulatorProxyRefType(const FString& InName, const uint32 InId)
+			: TModulatorProxyBase<IdType>(InName, InId)
+			, RefCount(0)
 		{
 		}
 
-		virtual ~TModulatorProxyRefBase() = default;
+		TModulatorProxyRefType(const TModulatorProxyRefType& InProxyRef)
+			: TModulatorProxyBase<IdType>(InProxyRef.GetName(), InProxyRef.GetId())
+			, RefCount(InProxyRef.RefCount)
+		{
+		}
 
-		virtual bool CanDestroy() const
+		TModulatorProxyRefType(TModulatorProxyRefType&& InProxyRef)
+			: TModulatorProxyBase<IdType>(InProxyRef.GetName(), InProxyRef.GetId())
+			, RefCount(InProxyRef.RefCount)
 		{
-			return !bAutoActivate || (bAutoActivate && RefSounds.Num() == 0);
+			InProxyRef.RefCount = 0;
 		}
 
-		int32 OnReleaseSound(const ISoundModulatable& Sound)
+		TModulatorProxyRefType& operator=(const TModulatorProxyRefType& InOther)
 		{
-			if (bAutoActivate)
-			{
-				const int32 NumRemoved = RefSounds.RemoveSingleSwap(&Sound, false /* bAllowShrinking */);
-				check(NumRemoved == 1);
-			}
+			RefCount = 0;
+			return *this;
+		}
 
-			return RefSounds.Num();
+		TModulatorProxyRefType& operator=(const TModulatorProxyRefType&& InOther)
+		{
+			RefCount = InOther.RefCount;
+			InOther.RefCount = 0;
+			return *this;
 		}
 
-		bool GetAutoActivate() const
+		virtual ~TModulatorProxyRefType()
 		{
-			return static_cast<bool>(bAutoActivate);
+			check(RefCount == 0);
 		}
 
-		const TArray<const ISoundModulatable*>& GetRefSounds() const
+		uint32 GetRefCount() const
 		{
-			return RefSounds;
+			return RefCount;
 		}
 
-		int32 OnInitSound(const ISoundModulatable& Sound)
+	protected:
+		FORCEINLINE void IncRef()
 		{
-			// Preview sounds force proxies into being auto-activated
-			// to allow for auditioning with the provided modulation settings.
-			if (Sound.IsPreviewSound())
-			{
-				bAutoActivate = 1;
-			}
+			RefCount++;
+		}
 
-			if (bAutoActivate)
-			{
-				RefSounds.Add(&Sound);
-			}
+		FORCEINLINE uint32 DecRef()
+		{
+			check(RefCount != 0);
 
-			return RefSounds.Num();
+			RefCount--;
+			return RefCount;
 		}
 
 	private:
-		uint8 bAutoActivate : 1;
-		TArray<const ISoundModulatable*> RefSounds;
+		friend class TProxyHandle<IdType, ProxyType, ProxyUObjType>;
 	};
 
-	struct FModulatorBusMixChannelProxy : public TModulatorProxyBase<FBusId>
-	{
-		FModulatorBusMixChannelProxy(const FSoundControlBusMixChannel& Channel);
-		FString Address;
-		uint32 ClassId;
-		FSoundModulationValue Value;
-	};
-
-	class FModulatorLFOProxy : public TModulatorProxyRefBase<FLFOId>
+	class FModulatorLFOProxy : public TModulatorProxyRefType<FLFOId, FModulatorLFOProxy, USoundBusModulatorLFO>
 	{
 	public:
 		FModulatorLFOProxy();
 		FModulatorLFOProxy(const USoundBusModulatorLFO& InLFO);
 
-		void OnUpdateProxy(const FModulatorLFOProxy& InLFOProxy);
+		FModulatorLFOProxy& operator =(const USoundBusModulatorLFO& InLFO);
 
 		float GetValue() const;
+		bool IsBypassed() const;
 		void Update(float InElapsed);
 
 	private:
+		void Init(const USoundBusModulatorLFO& InLFO);
+
 		Audio::FLFO LFO;
 		float Offset;
 		float Value;
+		uint8 bBypass : 1;
 	};
-	using LFOProxyMap = TMap<FBusId, FModulatorLFOProxy>;
 
-	class FControlBusProxy : public TModulatorProxyRefBase<FBusId>
+	class FControlBusProxy : public TModulatorProxyRefType<FBusId, FControlBusProxy, USoundControlBusBase>
 	{
 	public:
 		FControlBusProxy();
 		FControlBusProxy(const USoundControlBusBase& Bus);
 
-		void OnUpdateProxy(const FControlBusProxy& InBusProxy);
+		FControlBusProxy& operator =(const USoundControlBusBase& InLFO);
 
 		float GetDefaultValue() const;
-		const TArray<FLFOId>& GetLFOIds() const;
+		const TArray<FLFOHandle>& GetLFOHandles() const;
 		float GetLFOValue() const;
 		float GetMixValue() const;
 		FVector2D GetRange() const;
 		float GetValue() const;
+		void InitLFOs(const USoundControlBusBase& InBus, FLFOProxyMap& OutActiveLFOs);
+		bool IsBypassed() const;
 		void MixIn(const float InValue);
-		void MixLFO(LFOProxyMap& LFOMap);
+		void MixLFO();
 		void Reset();
 
 	private:
+		void Init(const USoundControlBusBase& InBus);
 		float Mix(float ValueA) const;
 		float Mix(float ValueA, float ValueB) const;
 
@@ -186,15 +385,17 @@ namespace AudioModulation
 		float LFOValue;
 		float MixValue;
 
-		TArray<FLFOId> LFOIds;
+		bool bBypass;
+
+		TArray<FLFOHandle> LFOHandles;
 		ESoundModulatorOperator Operator;
 		FVector2D Range;
 	};
 
-	using BusProxyMap = TMap<FBusId, FControlBusProxy>;
-
+	using FBusProxyMap = TMap<FBusId, FControlBusProxy>;
+	using FBusHandle = TProxyHandle<FBusId, FControlBusProxy, USoundControlBusBase>;
 
-	class FModulatorBusMixProxy : public TModulatorProxyRefBase<FBusMixId>
+	class FModulatorBusMixProxy : public TModulatorProxyRefType<FBusMixId, FModulatorBusMixProxy, USoundControlBusMix>
 	{
 	public:
 		enum class BusMixStatus : uint8
@@ -205,34 +406,35 @@ namespace AudioModulation
 		};
 
 		FModulatorBusMixProxy(const USoundControlBusMix& Mix);
-		virtual ~FModulatorBusMixProxy() = default;
 
-		virtual bool CanDestroy() const override;
+		FModulatorBusMixProxy& operator =(const USoundControlBusMix& InBusMix);
 
-		virtual void OnUpdateProxy(const FModulatorBusMixProxy& InBusMixProxy);
+		bool CanDestroy() const;
 
 		void SetEnabled();
 		void SetMix(const TArray<FSoundControlBusMixChannel>& InChannels);
 		void SetMixByFilter(const FString& InAddressFilter, uint32 InFilterClassId, const FSoundModulationValue& InValue);
 		void SetStopping();
 
-		void Update(const float Elapsed, BusProxyMap& ProxyMap);
+		void Update(const float Elapsed, FBusProxyMap& ProxyMap);
 
 		TMap<FBusId, FModulatorBusMixChannelProxy> Channels;
 
 	private:
+		void Init(const USoundControlBusMix& InBusMix);
 		BusMixStatus Status;
 	};
 
-	using BusMixProxyMap = TMap<FBusMixId, FModulatorBusMixProxy>;
+	using FBusMixProxyMap = TMap<FBusMixId, FModulatorBusMixProxy>;
+	using FBusMixHandle = TProxyHandle<FBusMixId, FModulatorBusMixProxy, USoundControlBusMix>;
 
 	/** Modulation input instance */
 	struct FModulationInputProxy
 	{
 		FModulationInputProxy();
-		FModulationInputProxy(const FSoundModulationInputBase& Patch);
+		FModulationInputProxy(const FSoundModulationInputBase& Patch, FReferencedProxies& OutRefProxies);
 
-		FBusId BusId;
+		FBusHandle BusHandle;
 		FSoundModulationInputTransform Transform;
 		uint8 bSampleAndHold : 1;
 	};
@@ -259,7 +461,7 @@ namespace AudioModulation
 	struct FModulationPatchProxy
 	{
 		FModulationPatchProxy();
-		FModulationPatchProxy(const FSoundModulationPatchBase& Patch);
+		FModulationPatchProxy(const FSoundModulationPatchBase& Patch, FReferencedProxies& OutRefProxies);
 
 		/** Default value of patch (Value mixed when inputs are provided or not, regardless of active state)*/
 		float DefaultInputValue;
@@ -277,7 +479,7 @@ namespace AudioModulation
 	struct FModulationSettingsProxy : public TModulatorProxyBase<uint32>
 	{
 		FModulationSettingsProxy();
-		FModulationSettingsProxy(const USoundModulationSettings& Settings);
+		FModulationSettingsProxy(const USoundModulationSettings& Settings, FReferencedProxies& OutRefProxies);
 
 		FModulationPatchProxy Volume;
 		FModulationPatchProxy Pitch;
@@ -286,6 +488,13 @@ namespace AudioModulation
 
 		TMap<FName, FModulationPatchProxy> Controls;
 
-		TArray<FBusMixId> Mixes;
+		TArray<FBusMixHandle> Mixes;
+	};
+
+	struct FReferencedProxies
+	{
+		FBusMixProxyMap BusMixes;
+		FBusProxyMap    Buses;
+		FLFOProxyMap    LFOs;
 	};
 } // namespace AudioModulation
\ No newline at end of file
diff --git a/Engine/Plugins/Runtime/AudioModulation/Source/AudioModulation/Private/SoundModulatorBase.cpp b/Engine/Plugins/Runtime/AudioModulation/Source/AudioModulation/Private/SoundModulatorBase.cpp
--- a/Engine/Plugins/Runtime/AudioModulation/Source/AudioModulation/Private/SoundModulatorBase.cpp
+++ b/Engine/Plugins/Runtime/AudioModulation/Source/AudioModulation/Private/SoundModulatorBase.cpp
@@ -11,31 +11,9 @@
 
 USoundModulatorBase::USoundModulatorBase(const FObjectInitializer& ObjectInitializer)
 	: Super(ObjectInitializer)
-	, bAutoActivate(0)
 {
 }
 
-bool USoundModulatorBase::CanAutoActivate(const ISoundModulatable* InSound) const
-{
-	if (bAutoActivate)
-	{
-		if (InSound)
-		{
-			return true;
-		}
-	}
-	else
-	{
-		if (!InSound || InSound->IsPreviewSound())
-		{
-			return true;
-		}
-	}
-
-	return false;
-}
-
-
 #if WITH_EDITOR
 void USoundModulatorBase::PostEditChangeProperty(struct FPropertyChangedEvent& PropertyChangedEvent)
 {
@@ -46,14 +24,6 @@ void USoundModulatorBase::PostEditChangeProperty(struct FPropertyChangedEvent& P
 		return;
 	}
 
-	// Cannot change AutoActivate while modulator is active to avoid leaking sound count anyway,
-	// so early out to avoid passing along a no-op
-	const FName Name = PropertyChangedEvent.Property->GetFName();
-	if (Name == TEXT("bAutoActivate"))
-	{
-		return;
-	}
-
 	FAudioDeviceManager* DeviceManager = GEngine->GetAudioDeviceManager();
 	if (!DeviceManager)
 	{
diff --git a/Engine/Plugins/Runtime/AudioModulation/Source/AudioModulation/Private/SoundModulatorLFO.cpp b/Engine/Plugins/Runtime/AudioModulation/Source/AudioModulation/Private/SoundModulatorLFO.cpp
--- a/Engine/Plugins/Runtime/AudioModulation/Source/AudioModulation/Private/SoundModulatorLFO.cpp
+++ b/Engine/Plugins/Runtime/AudioModulation/Source/AudioModulation/Private/SoundModulatorLFO.cpp
@@ -6,6 +6,7 @@
 #include "AudioModulationInternal.h"
 #include "AudioModulationLogging.h"
 #include "Engine/World.h"
+#include "SoundModulationProxy.h"
 
 
 USoundBusModulatorLFO::USoundBusModulatorLFO(const FObjectInitializer& ObjectInitializer)
@@ -15,6 +16,7 @@ USoundBusModulatorLFO::USoundBusModulatorLFO(const FObjectInitializer& ObjectIni
 	, Frequency(1.0f)
 	, Offset(0.5f)
 	, bLooping(1)
+	, bBypass(0)
 {
 }
 
@@ -28,20 +30,15 @@ void USoundBusModulatorLFO::BeginDestroy()
 		return;
 	}
 
-	FAudioDevice* AudioDevice = World->GetAudioDevice();
-	if (!AudioDevice)
+	if (FAudioDevice* AudioDevice = World->GetAudioDevice())
 	{
-		return;
-	}
-
-	check(AudioDevice->IsModulationPluginEnabled());
-	if (IAudioModulation* ModulationInterface = AudioDevice->ModulationInterface.Get())
-	{
-		auto ModulationImpl = static_cast<AudioModulation::FAudioModulation*>(ModulationInterface)->GetImpl();
-		check(ModulationImpl);
-
-		auto LFOId = static_cast<const AudioModulation::FLFOId>(GetUniqueID());
-		ModulationImpl->DeactivateLFO(LFOId);
+		check(AudioDevice->IsModulationPluginEnabled());
+		if (IAudioModulation* ModulationInterface = AudioDevice->ModulationInterface.Get())
+		{
+			auto ModulationImpl = static_cast<AudioModulation::FAudioModulation*>(ModulationInterface)->GetImpl();
+			check(ModulationImpl);
+			ModulationImpl->DeactivateLFO(*this);
+		}
 	}
 }
 
@@ -51,36 +48,50 @@ namespace AudioModulation
 	FModulatorLFOProxy::FModulatorLFOProxy()
 		: Offset(0.0f)
 		, Value(1.0f)
+		, bBypass(0)
 	{
 		LFO.SetFrequency(1.0f);
 		LFO.Start();
 	}
 
 	FModulatorLFOProxy::FModulatorLFOProxy(const USoundBusModulatorLFO& InLFO)
-		: TModulatorProxyRefBase<FLFOId>(InLFO.GetName(), InLFO.GetUniqueID(), InLFO.bAutoActivate)
+		: TModulatorProxyRefType(InLFO.GetName(), InLFO.GetUniqueID())
 		, Offset(InLFO.Offset)
 		, Value(1.0f)
+		, bBypass(InLFO.bBypass)
 	{
+		Init(InLFO);
+	}
+
+	FModulatorLFOProxy& FModulatorLFOProxy::operator =(const USoundBusModulatorLFO& InLFO)
+	{
+		Init(InLFO);
+		return *this;
+	}
+
+	float FModulatorLFOProxy::GetValue() const
+	{
+		return Value;
+	}
+
+	void FModulatorLFOProxy::Init(const USoundBusModulatorLFO& InLFO)
+	{
+		Offset = InLFO.Offset;
+		Value = 1.0f;
+		bBypass = InLFO.bBypass;
+
 		LFO.SetGain(InLFO.Amplitude);
 		LFO.SetFrequency(InLFO.Frequency);
 		LFO.SetMode(InLFO.bLooping ? Audio::ELFOMode::Type::Sync : Audio::ELFOMode::OneShot);
 
 		static_assert(static_cast<int32>(ESoundModulatorLFOShape::COUNT) == static_cast<int32>(Audio::ELFO::Type::NumLFOTypes), "LFOShape/ELFO Type mismatch");
 		LFO.SetType(static_cast<Audio::ELFO::Type>(InLFO.Shape));
-
 		LFO.Start();
 	}
 
-	float FModulatorLFOProxy::GetValue() const
-	{
-		return Value;
-	}
-
-	void FModulatorLFOProxy::OnUpdateProxy(const FModulatorLFOProxy& InLFOProxy)
+	bool FModulatorLFOProxy::IsBypassed() const
 	{
-		LFO = InLFOProxy.LFO;
-		Offset = InLFOProxy.Offset;
-		Value = InLFOProxy.Value;
+		return bBypass != 0;
 	}
 
 	void FModulatorLFOProxy::Update(float InElapsed)
diff --git a/Engine/Plugins/Runtime/AudioModulation/Source/AudioModulation/Public/AudioModulation.h b/Engine/Plugins/Runtime/AudioModulation/Source/AudioModulation/Public/AudioModulation.h
--- a/Engine/Plugins/Runtime/AudioModulation/Source/AudioModulation/Public/AudioModulation.h
+++ b/Engine/Plugins/Runtime/AudioModulation/Source/AudioModulation/Public/AudioModulation.h
@@ -27,7 +27,7 @@ namespace AudioModulation
 		virtual ~FAudioModulation() = default;
 
 		//~ Begin IAudioModulation implementation
-		virtual float CalculateInitialVolume(const USoundModulationPluginSourceSettingsBase& InSettingsBase) const override;
+		virtual float CalculateInitialVolume(const USoundModulationPluginSourceSettingsBase& InSettingsBase) override;
 
 		virtual void Initialize(const FAudioPluginInitializationParams& InitializationParams) override;
 		virtual void OnInitSound(ISoundModulatable& Sound, const USoundModulationPluginSourceSettingsBase& Settings) override;
diff --git a/Engine/Plugins/Runtime/AudioModulation/Source/AudioModulation/Public/SoundControlBus.h b/Engine/Plugins/Runtime/AudioModulation/Source/AudioModulation/Public/SoundControlBus.h
--- a/Engine/Plugins/Runtime/AudioModulation/Source/AudioModulation/Public/SoundControlBus.h
+++ b/Engine/Plugins/Runtime/AudioModulation/Source/AudioModulation/Public/SoundControlBus.h
@@ -33,14 +33,18 @@ class USoundControlBusBase : public USoundModulatorBase
 	GENERATED_UCLASS_BODY()
 
 public:
+	/** If true, bypasses control bus from being mixed. */
+	UPROPERTY(EditAnywhere, Category = General, BlueprintReadWrite)
+	bool bBypass;
+
 #if WITH_EDITORONLY_DATA
 	/** If true, Address field is used in place of object name for address used when applying mix changes using filtering. */
-	UPROPERTY(EditAnywhere, Category = General, BlueprintReadWrite)
+	UPROPERTY(EditAnywhere, Category = Mix, BlueprintReadWrite)
 	bool bOverrideAddress;
 #endif // WITH_EDITORONLY_DATA
 
 	/** Address to use when applying mix changes. */
-	UPROPERTY(EditAnywhere, Category = General, BlueprintReadWrite, meta = (EditCondition = "bOverrideAddress"))
+	UPROPERTY(EditAnywhere, Category = Mix, BlueprintReadWrite, meta = (EditCondition = "bOverrideAddress"))
 	FString Address;
 
 	/** Default value of modulator (when no mix is applied). */
diff --git a/Engine/Plugins/Runtime/AudioModulation/Source/AudioModulation/Public/SoundModulationPatch.h b/Engine/Plugins/Runtime/AudioModulation/Source/AudioModulation/Public/SoundModulationPatch.h
--- a/Engine/Plugins/Runtime/AudioModulation/Source/AudioModulation/Public/SoundModulationPatch.h
+++ b/Engine/Plugins/Runtime/AudioModulation/Source/AudioModulation/Public/SoundModulationPatch.h
@@ -79,7 +79,7 @@ struct FSoundModulationInputBase
 	uint8 bSampleAndHold : 1;
 
 	/** Transform to apply to the input prior to mix phase */
-	UPROPERTY(EditAnywhere, BlueprintReadWrite, Category = Input, meta = (ShowOnlyInnerProperties))
+	UPROPERTY(EditAnywhere, BlueprintReadWrite, Category = Advanced)
 	FSoundModulationInputTransform Transform;
 
 	virtual const USoundControlBusBase* GetBus() const PURE_VIRTUAL(FSoundModulationInputBase::GetBus, return nullptr; );
@@ -172,7 +172,7 @@ struct FSoundModulationPatchBase
 	UPROPERTY(EditAnywhere, BlueprintReadWrite, Category = Inputs, meta = (UIMin = "0", UIMax = "1"))
 	uint8 bBypass : 1;
 
-	virtual void GenerateProxies(TArray<AudioModulation::FModulationInputProxy>& InputProxies) const PURE_VIRTUAL(FSoundModulationPatchBase::GenerateProxies, );
+	virtual TArray<const FSoundModulationInputBase*> GetInputs() const PURE_VIRTUAL(FSoundModulationPatchBase::GetInputs(), return TArray<const FSoundModulationInputBase*>(););
 	virtual FSoundModulationOutputBase* GetOutput() PURE_VIRTUAL(FSoundModulationPatchBase::GetOutput(), return nullptr;);
 	virtual const FSoundModulationOutputBase* GetOutput() const PURE_VIRTUAL(FSoundModulationPatchBase::GetOutput(), return nullptr;);
 
@@ -194,7 +194,7 @@ struct FSoundVolumeModulationPatch : public FSoundModulationPatchBase
 	UPROPERTY(EditAnywhere, BlueprintReadWrite, Category = Output, meta = (ShowOnlyInnerProperties))
 	FSoundModulationOutputFixedOperator Output;
 
-	virtual void GenerateProxies(TArray<AudioModulation::FModulationInputProxy>& InputProxies) const override;
+	virtual TArray<const FSoundModulationInputBase*> GetInputs() const override;
 	virtual FSoundModulationOutputBase* GetOutput() override { return static_cast<FSoundModulationOutputBase*>(&Output); }
 	virtual const FSoundModulationOutputBase* GetOutput() const override { return static_cast<const FSoundModulationOutputBase*>(&Output); }
 
@@ -216,7 +216,7 @@ struct FSoundPitchModulationPatch : public FSoundModulationPatchBase
 	UPROPERTY(EditAnywhere, BlueprintReadWrite, Category = Output, meta = (ShowOnlyInnerProperties))
 	FSoundModulationOutputFixedOperator Output;
 
-	virtual void GenerateProxies(TArray<AudioModulation::FModulationInputProxy>& InputProxies) const override;
+	virtual TArray<const FSoundModulationInputBase*> GetInputs() const override;
 	virtual FSoundModulationOutputBase* GetOutput() override { return static_cast<FSoundModulationOutputBase*>(&Output); }
 	virtual const FSoundModulationOutputBase* GetOutput() const override { return static_cast<const FSoundModulationOutputBase*>(&Output); }
 
@@ -238,7 +238,7 @@ struct FSoundLPFModulationPatch : public FSoundModulationPatchBase
 	UPROPERTY(EditAnywhere, BlueprintReadWrite, Category = Output, meta = (ShowOnlyInnerProperties))
 	FSoundModulationOutputFixedOperator Output;
 
-	virtual void GenerateProxies(TArray<AudioModulation::FModulationInputProxy>& InputProxies) const override;
+	virtual TArray<const FSoundModulationInputBase*> GetInputs() const override;
 	virtual FSoundModulationOutputBase* GetOutput() override { return static_cast<FSoundModulationOutputBase*>(&Output); }
 	virtual const FSoundModulationOutputBase* GetOutput() const override { return static_cast<const FSoundModulationOutputBase*>(&Output); }
 
@@ -260,7 +260,7 @@ struct FSoundHPFModulationPatch : public FSoundModulationPatchBase
 	UPROPERTY(EditAnywhere, BlueprintReadWrite, Category = Output, meta = (ShowOnlyInnerProperties))
 	FSoundModulationOutputFixedOperator Output;
 
-	virtual void GenerateProxies(TArray<AudioModulation::FModulationInputProxy>& InputProxies) const override;
+	virtual TArray<const FSoundModulationInputBase*> GetInputs() const override;
 	virtual FSoundModulationOutputBase* GetOutput() override { return static_cast<FSoundModulationOutputBase*>(&Output); }
 	virtual const FSoundModulationOutputBase* GetOutput() const override { return static_cast<const FSoundModulationOutputBase*>(&Output); }
 
@@ -286,7 +286,7 @@ struct FSoundControlModulationPatch : public FSoundModulationPatchBase
 	UPROPERTY(EditAnywhere, BlueprintReadWrite, Category = Output, meta = (ShowOnlyInnerProperties))
 	FSoundModulationOutput Output;
 
-	virtual void GenerateProxies(TArray<AudioModulation::FModulationInputProxy>& InputProxies) const override;
+	virtual TArray<const FSoundModulationInputBase*> GetInputs() const override;
 	virtual FSoundModulationOutputBase* GetOutput() override { return static_cast<FSoundModulationOutputBase*>(&Output); }
 	virtual const FSoundModulationOutputBase* GetOutput() const override { return static_cast<const FSoundModulationOutputBase*>(&Output); }
 };
diff --git a/Engine/Plugins/Runtime/AudioModulation/Source/AudioModulation/Public/SoundModulatorBase.h b/Engine/Plugins/Runtime/AudioModulation/Source/AudioModulation/Public/SoundModulatorBase.h
--- a/Engine/Plugins/Runtime/AudioModulation/Source/AudioModulation/Public/SoundModulatorBase.h
+++ b/Engine/Plugins/Runtime/AudioModulation/Source/AudioModulation/Public/SoundModulatorBase.h
@@ -19,15 +19,6 @@ class USoundModulatorBase : public UObject
 {
 	GENERATED_UCLASS_BODY()
 
-public:
-	/**
-	 * When set, automatically activates/deactivates when sounds referencing asset begin/end playing respectively.
-	 */
-	UPROPERTY(EditAnywhere, Category = General, BlueprintReadWrite)
-	uint8 bAutoActivate : 1;
-
-	bool CanAutoActivate(const ISoundModulatable* InSound) const;
-
 #if WITH_EDITOR
 	virtual void PostEditChangeProperty(FPropertyChangedEvent& PropertyChangedEvent) override;
 #endif // WITH_EDITOR
diff --git a/Engine/Plugins/Runtime/AudioModulation/Source/AudioModulation/Public/SoundModulatorLFO.h b/Engine/Plugins/Runtime/AudioModulation/Source/AudioModulation/Public/SoundModulatorLFO.h
--- a/Engine/Plugins/Runtime/AudioModulation/Source/AudioModulation/Public/SoundModulatorLFO.h
+++ b/Engine/Plugins/Runtime/AudioModulation/Source/AudioModulation/Public/SoundModulatorLFO.h
@@ -52,4 +52,8 @@ public:
 	/** Whether or not to loop the oscillation more than once */
 	UPROPERTY(EditAnywhere, Category = Modulation, BlueprintReadWrite)
 	uint8 bLooping : 1;
+
+	/** Whether or not to bypass the LFO from being included in mix. */
+	UPROPERTY(EditAnywhere, Category = Modulation, BlueprintReadWrite)
+	uint8 bBypass: 1;
 };
diff --git a/Engine/Source/Runtime/Engine/Public/IAudioExtensionPlugin.h b/Engine/Source/Runtime/Engine/Public/IAudioExtensionPlugin.h
--- a/Engine/Source/Runtime/Engine/Public/IAudioExtensionPlugin.h
+++ b/Engine/Source/Runtime/Engine/Public/IAudioExtensionPlugin.h
@@ -664,7 +664,7 @@ public:
 	virtual ~IAudioModulation() { }
 
 	/** Calculates initial volume to determine if sound is audible using base settings data */
-	virtual float CalculateInitialVolume(const USoundModulationPluginSourceSettingsBase& InSettingsBase) const { return 1.0f; }
+	virtual float CalculateInitialVolume(const USoundModulationPluginSourceSettingsBase& InSettingsBase) { return 1.0f; }
 
 	/** Initialize the modulation plugin with the same rate and number of sources */
 	virtual void Initialize(const FAudioPluginInitializationParams& InitializationParams) { }
