Fix display issue with MixByFilter static BP function call

diff --git a/Engine/Plugins/Runtime/AudioModulation/Source/AudioModulation/Private/AudioModulationStatics.h b/Engine/Plugins/Runtime/AudioModulation/Source/AudioModulation/Private/AudioModulationStatics.h
--- a/Engine/Plugins/Runtime/AudioModulation/Source/AudioModulation/Private/AudioModulationStatics.h
+++ b/Engine/Plugins/Runtime/AudioModulation/Source/AudioModulation/Private/AudioModulationStatics.h
@@ -144,10 +144,10 @@ public:
 	 * @param AddressFilter - Address filter to apply to provided mix's channels.
 	 * @param BusClass - Filters buses by subclass.
 	 * @param Value - Target value to mix filtered channels to.
-	 * @param AttackTime - If positive, updates the attack time for the resulting bus channels found matching the provided filter.
-	 * @param ReleaseTime - If positive, updates the release time for the resulting bus channels found matching the provided filter.
+	 * @param AttackTime - If non-negative, updates the attack time for the resulting bus channels found matching the provided filter.
+	 * @param ReleaseTime - If non-negative, updates the release time for the resulting bus channels found matching the provided filter.
 	 */
-	UFUNCTION(BlueprintCallable, Category = "Audio", DisplayName = "Update Control Bus Mix By Filter", meta = (AdvancedDisplay = "4", WorldContext = "WorldContextObject", Keywords = "update bus control class modulation modulator mix channel value filter"))
+	UFUNCTION(BlueprintCallable, Category = "Audio", DisplayName = "Update Control Bus Mix By Filter", meta = (AdvancedDisplay = "5", WorldContext = "WorldContextObject", Keywords = "update bus control class modulation modulator mix channel value filter"))
 	static void UpdateMixByFilter(
 		const UObject*						WorldContextObject,
 		USoundControlBusMix*				Mix,
