Fix shadow variable issue on Lumin (trivial)

diff --git a/Engine/Plugins/Runtime/AudioModulation/Source/AudioModulation/Private/AudioModulationImpl.cpp b/Engine/Plugins/Runtime/AudioModulation/Source/AudioModulation/Private/AudioModulationImpl.cpp
--- a/Engine/Plugins/Runtime/AudioModulation/Source/AudioModulation/Private/AudioModulationImpl.cpp
+++ b/Engine/Plugins/Runtime/AudioModulation/Source/AudioModulation/Private/AudioModulationImpl.cpp
@@ -494,9 +494,9 @@ namespace AudioModulation
 
 	void FAudioModulationImpl::UpdateModulator(const USoundModulatorBase& InModulator)
 	{
-		if (const USoundBusModulatorLFO* LFO = Cast<USoundBusModulatorLFO>(&InModulator))
+		if (const USoundBusModulatorLFO* InLFO = Cast<USoundBusModulatorLFO>(&InModulator))
 		{
-			const TWeakObjectPtr<const USoundBusModulatorLFO> LFOPtr = LFO;
+			const TWeakObjectPtr<const USoundBusModulatorLFO> LFOPtr = InLFO;
 			RunCommandOnAudioThread([this, LFOPtr]()
 			{
 				if (LFOPtr.IsValid())
@@ -508,9 +508,9 @@ namespace AudioModulation
 			});
 		}
 
-		if (const USoundControlBusBase* Bus = Cast<USoundControlBusBase>(&InModulator))
+		if (const USoundControlBusBase* InBus = Cast<USoundControlBusBase>(&InModulator))
 		{
-			const TWeakObjectPtr<const USoundControlBusBase> BusPtr = Bus;
+			const TWeakObjectPtr<const USoundControlBusBase> BusPtr = InBus;
 			RunCommandOnAudioThread([this, BusPtr]()
 			{
 				if (BusPtr.IsValid())
@@ -524,16 +524,16 @@ namespace AudioModulation
 			});
 		}
 
-		if (const USoundControlBusMix* Mix = Cast<USoundControlBusMix>(&InModulator))
+		if (const USoundControlBusMix* InMix = Cast<USoundControlBusMix>(&InModulator))
 		{
-			const TWeakObjectPtr<const USoundControlBusMix> BusMixPtr = Mix;
+			const TWeakObjectPtr<const USoundControlBusMix> BusMixPtr = InMix;
 			RunCommandOnAudioThread([this, BusMixPtr]()
 			{
 				if (BusMixPtr.IsValid())
 				{
-					const USoundControlBusMix& BusMix = *BusMixPtr.Get();
-					FBusMixHandle BusMixHandle = FBusMixHandle::Create(BusMix, RefProxies.BusMixes);
-					BusMixHandle.FindProxy() = BusMix;
+					const USoundControlBusMix& Mix = *BusMixPtr.Get();
+					FBusMixHandle BusMixHandle = FBusMixHandle::Create(Mix, RefProxies.BusMixes);
+					BusMixHandle.FindProxy() = Mix;
 				}
 			});
 		}
