Perf Fix: Avoid creating modulation proxies every time an update request comes in if the proxy is not already active (as it will simply delete the proxy at the end of the call, and update calls are anticipated to be called relatively frequently)

diff --git a/Engine/Plugins/Runtime/AudioModulation/Source/AudioModulation/Private/AudioModulationImpl.cpp b/Engine/Plugins/Runtime/AudioModulation/Source/AudioModulation/Private/AudioModulationImpl.cpp
--- a/Engine/Plugins/Runtime/AudioModulation/Source/AudioModulation/Private/AudioModulationImpl.cpp
+++ b/Engine/Plugins/Runtime/AudioModulation/Source/AudioModulation/Private/AudioModulationImpl.cpp
@@ -502,8 +502,15 @@ namespace AudioModulation
 				if (LFOPtr.IsValid())
 				{
 					const USoundBusModulatorLFO& LFO = *LFOPtr.Get();
-					FLFOHandle LFOHandle = FLFOHandle::Create(LFO, RefProxies.LFOs);
-					LFOHandle.FindProxy() = LFO;
+					FLFOHandle LFOHandle = FLFOHandle::Get(LFO, RefProxies.LFOs);
+					if (LFOHandle.IsValid())
+					{
+						LFOHandle.FindProxy() = LFO;
+					}
+					else
+					{
+						UE_LOG(LogAudioModulation, Verbose, TEXT("Update to '%s' Ignored: LFO is inactive."), *LFO.GetName());
+					}
 				}
 			});
 		}
@@ -516,10 +523,17 @@ namespace AudioModulation
 				if (BusPtr.IsValid())
 				{
 					const USoundControlBusBase* Bus = BusPtr.Get();
-					FBusHandle BusHandle = FBusHandle::Create(*Bus, RefProxies.Buses);
-					FControlBusProxy& Proxy = BusHandle.FindProxy();
-					Proxy = *Bus;
-					Proxy.InitLFOs(*Bus, RefProxies.LFOs);
+					FBusHandle BusHandle = FBusHandle::Get(*Bus, RefProxies.Buses);
+					if (BusHandle.IsValid())
+					{
+						FControlBusProxy& Proxy = BusHandle.FindProxy();
+						Proxy = *Bus;
+						Proxy.InitLFOs(*Bus, RefProxies.LFOs);
+					}
+					else
+					{
+						UE_LOG(LogAudioModulation, Verbose, TEXT("Update to '%s' Ignored: Control Bus is inactive."), *Bus->GetName());
+					}
 				}
 			});
 		}
@@ -532,8 +546,15 @@ namespace AudioModulation
 				if (BusMixPtr.IsValid())
 				{
 					const USoundControlBusMix& Mix = *BusMixPtr.Get();
-					FBusMixHandle BusMixHandle = FBusMixHandle::Create(Mix, RefProxies.BusMixes);
-					BusMixHandle.FindProxy() = Mix;
+					FBusMixHandle BusMixHandle = FBusMixHandle::Get(Mix, RefProxies.BusMixes);
+					if (BusMixHandle.IsValid())
+					{
+						BusMixHandle.FindProxy() = Mix;
+					}
+					else
+					{
+						UE_LOG(LogAudioModulation, Verbose, TEXT("Update to '%s' Ignored: Control Bus Mix is inactive."), *Mix.GetName());
+					}
 				}
 			});
 		}
diff --git a/Engine/Plugins/Runtime/AudioModulation/Source/AudioModulation/Private/SoundModulationProxy.h b/Engine/Plugins/Runtime/AudioModulation/Source/AudioModulation/Private/SoundModulationProxy.h
--- a/Engine/Plugins/Runtime/AudioModulation/Source/AudioModulation/Private/SoundModulationProxy.h
+++ b/Engine/Plugins/Runtime/AudioModulation/Source/AudioModulation/Private/SoundModulationProxy.h
@@ -121,7 +121,24 @@ namespace AudioModulation
 		}
 
 		/*
-		 * Creates a handle to a proxy modulation object tracked in the provided InProxyMap
+		 * Creates a handle to a proxy modulation object tracked in the provided InProxyMap if it exists, otherwise returns invalid handle.
+		 */
+		static TProxyHandle<IdType, ProxyType, ProxyUObjType> Get(const ProxyUObjType& InObject, TMap<IdType, ProxyType>& InProxyMap)
+		{
+			const IdType ObjectId = static_cast<IdType>(InObject.GetUniqueID());
+			if (ProxyType* Proxy = InProxyMap.Find(ObjectId))
+			{
+				TProxyHandle<IdType, ProxyType, ProxyUObjType> NewHandle(ObjectId, InProxyMap);
+				Proxy->IncRef();
+
+				return NewHandle;
+			} 
+
+			return TProxyHandle<IdType, ProxyType, ProxyUObjType>();
+		}
+
+		/*
+		 * Creates a handle to a proxy modulation object tracked in the provided InProxyMap.  Creates new proxy if it doesn't exist.
 		 */
 		static TProxyHandle<IdType, ProxyType, ProxyUObjType> Create(const ProxyUObjType& InObject, TMap<IdType, ProxyType>& InProxyMap)
 		{
