Clean-up UX for modulation bus types (descriptions, order, advanced hiding)

diff --git a/Engine/Plugins/Runtime/AudioModulation/Source/AudioModulation/Public/SoundControlBus.h b/Engine/Plugins/Runtime/AudioModulation/Source/AudioModulation/Public/SoundControlBus.h
--- a/Engine/Plugins/Runtime/AudioModulation/Source/AudioModulation/Public/SoundControlBus.h
+++ b/Engine/Plugins/Runtime/AudioModulation/Source/AudioModulation/Public/SoundControlBus.h
@@ -37,6 +37,18 @@ public:
 	UPROPERTY(EditAnywhere, Category = General, BlueprintReadWrite)
 	bool bBypass;
 
+	/** Default value of modulator when no mix is applied. Value that is also returned to when mix is released. */
+	UPROPERTY(EditAnywhere, Category = General, BlueprintReadWrite, meta = (UIMin = "0", UIMax = "1"))
+	float DefaultValue;
+
+	/** Minimum value the bus can achieve (applied post mix phase, pre patch output) */
+	UPROPERTY(EditAnywhere, Category = General, BlueprintReadWrite, meta = (UIMin = "0", UIMax = "1"))
+	float Min;
+
+	/** Maximum value the bus can achieve (applied post mix phase, pre patch output) */
+	UPROPERTY(EditAnywhere, Category = General, BlueprintReadWrite, meta = (UIMin = "0", UIMax = "1"))
+	float Max;
+
 #if WITH_EDITORONLY_DATA
 	/** If true, Address field is used in place of object name for address used when applying mix changes using filtering. */
 	UPROPERTY(EditAnywhere, Category = Mix, BlueprintReadWrite)
@@ -47,18 +59,6 @@ public:
 	UPROPERTY(EditAnywhere, Category = Mix, BlueprintReadWrite, meta = (EditCondition = "bOverrideAddress"))
 	FString Address;
 
-	/** Default value of modulator (when no mix is applied). */
-	UPROPERTY(EditAnywhere, Category = Modulation, BlueprintReadWrite, meta = (UIMin = "0", UIMax = "1"))
-	float DefaultValue;
-
-	/** Minimum value the bus can achieve (applied post mix phase, pre patch output) */
-	UPROPERTY(EditAnywhere, Category = Modulation, BlueprintReadWrite, meta = (UIMin = "0", UIMax = "1"))
-	float Min;
-
-	/** Maximum value the bus can achieve (applied post mix phase, pre patch output) */
-	UPROPERTY(EditAnywhere, Category = Modulation, BlueprintReadWrite, meta = (UIMin = "0", UIMax = "1"))
-	float Max;
-
 	UPROPERTY(EditAnywhere, Category = Modulation, BlueprintReadWrite)
 	TArray<USoundBusModulatorBase*> Modulators;
 
diff --git a/Engine/Plugins/Runtime/AudioModulation/Source/AudioModulation/Public/SoundModulationPatch.h b/Engine/Plugins/Runtime/AudioModulation/Source/AudioModulation/Public/SoundModulationPatch.h
--- a/Engine/Plugins/Runtime/AudioModulation/Source/AudioModulation/Public/SoundModulationPatch.h
+++ b/Engine/Plugins/Runtime/AudioModulation/Source/AudioModulation/Public/SoundModulationPatch.h
@@ -79,7 +79,7 @@ struct FSoundModulationInputBase
 	uint8 bSampleAndHold : 1;
 
 	/** Transform to apply to the input prior to mix phase */
-	UPROPERTY(EditAnywhere, BlueprintReadWrite, Category = Advanced)
+	UPROPERTY(EditAnywhere, BlueprintReadWrite, Category = Input, meta = (DisplayName = "Advanced"))
 	FSoundModulationInputTransform Transform;
 
 	virtual const USoundControlBusBase* GetBus() const PURE_VIRTUAL(FSoundModulationInputBase::GetBus, return nullptr; );
diff --git a/Engine/Plugins/Runtime/AudioModulation/Source/AudioModulation/Public/SoundModulationTransform.h b/Engine/Plugins/Runtime/AudioModulation/Source/AudioModulation/Public/SoundModulationTransform.h
--- a/Engine/Plugins/Runtime/AudioModulation/Source/AudioModulation/Public/SoundModulationTransform.h
+++ b/Engine/Plugins/Runtime/AudioModulation/Source/AudioModulation/Public/SoundModulationTransform.h
@@ -16,19 +16,19 @@ struct FSoundModulationInputTransform
 
 	FSoundModulationInputTransform();
 
-	/** Minimum value to clamp the input to. */
+	/** Minimum value to clamp the input to prior to transforming via linear interpolation. */
 	UPROPERTY(EditAnywhere, Category = Modulation, BlueprintReadWrite, meta = (DisplayName = "Input Min", UIMin = "0", UIMax = "1"))
 	float InputMin;
 
-	/** Maximum value to clamp the input to. */
+	/** Maximum value to clamp the input to prior to transforming via linear interpolation. */
 	UPROPERTY(EditAnywhere, Category = Modulation, BlueprintReadWrite, meta = (DisplayName = "Input Max", UIMin = "0", UIMax = "1"))
 	float InputMax;
 
-	/** Minimum value to clamp the output to. */
+	/** Minimum value to scale the output to. */
 	UPROPERTY(EditAnywhere, Category = Modulation, BlueprintReadWrite, meta = (DisplayName = "Output Min", UIMin = "0", UIMax = "1"))
 	float OutputMin;
 
-	/** Maximum value to clamp the output to. */
+	/** Maximum value to scale the output to. */
 	UPROPERTY(EditAnywhere, Category = Modulation, BlueprintReadWrite, meta = (DisplayName = "Output Max", UIMin = "0", UIMax = "1"))
 	float OutputMax;
 
