Fix init order warning

diff --git a/Engine/Plugins/Runtime/AudioModulation/Source/AudioModulation/Private/SoundControlBus.cpp b/Engine/Plugins/Runtime/AudioModulation/Source/AudioModulation/Private/SoundControlBus.cpp
--- a/Engine/Plugins/Runtime/AudioModulation/Source/AudioModulation/Private/SoundControlBus.cpp
+++ b/Engine/Plugins/Runtime/AudioModulation/Source/AudioModulation/Private/SoundControlBus.cpp
@@ -12,12 +12,13 @@
 
 USoundControlBusBase::USoundControlBusBase(const FObjectInitializer& ObjectInitializer)
 	: Super(ObjectInitializer)
-#if WITH_EDITORONLY_DATA
-	, bOverrideAddress(false)
-#endif
+	, bBypass(false)
 	, DefaultValue(1.0f)
 	, Min(0.0f)
 	, Max(1.0f)
+#if WITH_EDITORONLY_DATA
+	, bOverrideAddress(false)
+#endif
 {
 }
 
