Disable offending clang compiler error in user-sanity bounds check.
#rnx

diff --git a/Engine/Plugins/Runtime/AudioModulation/Source/AudioModulationEditor/Private/Layouts/SoundModulationParameterLayout.cpp b/Engine/Plugins/Runtime/AudioModulation/Source/AudioModulationEditor/Private/Layouts/SoundModulationParameterLayout.cpp
--- a/Engine/Plugins/Runtime/AudioModulation/Source/AudioModulationEditor/Private/Layouts/SoundModulationParameterLayout.cpp
+++ b/Engine/Plugins/Runtime/AudioModulation/Source/AudioModulationEditor/Private/Layouts/SoundModulationParameterLayout.cpp
@@ -74,27 +74,28 @@ void FSoundModulationParameterLayoutCustomization::CustomizeChildren(TSharedRef<
 
 void FSoundModulationParameterLayoutCustomization::OnValueChanged(TSharedRef<IPropertyHandle> InStructPropertyHandle)
 {
-	// Potentially accessing the value while garbage collecting or saving the package could trigger a crash.
-	// so we fail to get the value when that is occurring.
-	if (GIsSavingPackage || IsGarbageCollecting())
-	{
-		return;
-	}
-
-	// Way to enforce value was set by editor and remains in-bounds but allows each callsite to set and check
-	// at runtime (as desired being that input bus values can be out-of-bounds by design)
-	if (InStructPropertyHandle->IsValidHandle())
-	{
-		TArray<void*> RawData;
-		InStructPropertyHandle->AccessRawData(RawData);
-		for (void* RawPtr : RawData)
-		{
-			if (RawPtr)
-			{
-				FSoundModulationParameter& ModParam = *reinterpret_cast<FSoundModulationParameter*>(RawPtr);
-				ModParam.SetBounds(ModParam.GetMinValue(), ModParam.GetMaxValue());
-			}
-		}
-	}
+// Disabling due to compile error in Clang. TODO: Find another way to clamp value at edit time that is safer.
+// 	// Potentially accessing the value while garbage collecting or saving the package could trigger a crash.
+// 	// so we fail to get the value when that is occurring.
+// 	if (GIsSavingPackage || IsGarbageCollecting())
+// 	{
+// 		return;
+// 	}
+// 
+// 	// Way to enforce value was set by editor and remains in-bounds but allows each callsite to set and check
+// 	// at runtime (as desired being that input bus values can be out-of-bounds by design)
+// 	if (InStructPropertyHandle->IsValidHandle())
+// 	{
+// 		TArray<void*> RawData;
+// 		InStructPropertyHandle->AccessRawData(RawData);
+// 		for (void* RawPtr : RawData)
+// 		{
+// 			if (RawPtr)
+// 			{
+// 				FSoundModulationParameter& ModParam = *reinterpret_cast<FSoundModulationParameter*>(RawPtr);
+// 				ModParam.SetBounds(ModParam.GetMinValue(), ModParam.GetMaxValue());
+// 			}
+// 		}
+// 	}
 }
 #undef LOCTEXT_NAMESPACE
