- Add Mix serialization/deserialization to ini files
- Fix mix set/update logic and crash involved with removing mix proxy before stop fade complete
- Mixes can now activate buses they act on if bus is not already active. This avoids a lot of unnecessary legwork to ensure active status when buses *should* be active.  It also avoids a lot of complexity around mix state being mismatched between UObject & proxy representation

diff --git a/Engine/Plugins/Runtime/AudioModulation/AudioModulation.uplugin b/Engine/Plugins/Runtime/AudioModulation/AudioModulation.uplugin
--- a/Engine/Plugins/Runtime/AudioModulation/AudioModulation.uplugin
+++ b/Engine/Plugins/Runtime/AudioModulation/AudioModulation.uplugin
@@ -1,9 +1,9 @@
 {
 	"FileVersion": 3,
 	"Version": 1,
-	"VersionName": "1.0",
+	"VersionName": "1.1",
 	"FriendlyName": "Audio Modulation",
-	"Description": "Plugin provides an interface for modulating common parameters in audio.",
+	"Description": "Enables modulation of control parameters in the Unreal Audio Engine.",
 	"Category": "Audio",
 	"CreatedBy": "Epic Games, Inc.",
 	"CreatedByURL": "http://epicgames.com",
@@ -14,24 +14,17 @@
 	"CanContainContent": false,
 	"IsBetaVersion": true,
 	"Installed": false,
-	"Modules": [
+	"Modules":
+	[
 		{
 			"Name": "AudioModulation",
 			"Type": "Runtime",
-			"LoadingPhase": "Default",
-			"WhitelistPlatforms": [
-				"Win64",
-				"Lumin",
-				"Mac"
-			]
+			"LoadingPhase": "Default"
 		},
 		{
 			"Name": "AudioModulationEditor",
 			"Type": "Editor",
-			"LoadingPhase": "PostEngineInit",
-			"WhitelistPlatforms": [
-				"Win64"
-			]
+			"LoadingPhase": "PostEngineInit"
 		}
 	]
 }
\ No newline at end of file
diff --git a/Engine/Plugins/Runtime/AudioModulation/Source/AudioModulation/Private/AudioModulationImpl.cpp b/Engine/Plugins/Runtime/AudioModulation/Source/AudioModulation/Private/AudioModulationImpl.cpp
--- a/Engine/Plugins/Runtime/AudioModulation/Source/AudioModulation/Private/AudioModulationImpl.cpp
+++ b/Engine/Plugins/Runtime/AudioModulation/Source/AudioModulation/Private/AudioModulationImpl.cpp
@@ -2,7 +2,9 @@
 #include "AudioModulationInternal.h"
 
 #if WITH_AUDIOMODULATION
+#include "Audio/AudioAddressPattern.h"
 #include "AudioModulationLogging.h"
+#include "AudioModulationProfileSerializer.h"
 #include "AudioThread.h"
 #include "Engine/Engine.h"
 #include "IAudioExtensionPlugin.h"
@@ -11,10 +13,13 @@
 #include "SoundModulationPatch.h"
 #include "SoundModulationProxy.h"
 #include "SoundModulationValue.h"
+#include "UObject/UObjectIterator.h"
 #include "UObject/WeakObjectPtr.h"
 
+
 #if !UE_BUILD_SHIPPING
 #include "AudioModulationDebugger.h"
+#include "Async/Async.h"
 #endif // !UE_BUILD_SHIPPING
 
 DECLARE_DWORD_COUNTER_STAT(TEXT("Bus Count"), STAT_AudioModulationBusCount, STATGROUP_AudioModulation)
@@ -22,7 +27,7 @@ DECLARE_DWORD_COUNTER_STAT(TEXT("LFO Count"), STAT_AudioModulationLFOCount, STAT
 DECLARE_DWORD_COUNTER_STAT(TEXT("Mix Count"), STAT_AudioModulationMixCount, STATGROUP_AudioModulation)
 
 
-namespace
+namespace AudioModulation
 {
 	void MixInModulationValue(ESoundModulatorOperator& Operator, float ModStageValue, float& Value)
 	{
@@ -49,11 +54,19 @@ namespace
 		break;
 		}
 	}
-} // namespace <>
 
+	struct FProfileChannelInfo
+	{
+		USoundControlBusBase* Bus;
+		FSoundModulationValue Value;
+
+		FProfileChannelInfo(const FModulatorBusMixChannelProxy& InProxy)
+			: Bus(nullptr)
+			, Value(InProxy.Value)
+		{
+		}
+	};
 
-namespace AudioModulation
-{
 	FAudioModulationImpl::FAudioModulationImpl()
 	{
 	}
@@ -178,7 +191,16 @@ namespace AudioModulation
 		{
 			if (BusMixPtr.IsValid())
 			{
-				FBusMixHandle BusMixHandle = FBusMixHandle::Create(*BusMixPtr.Get(), RefProxies.BusMixes);
+				FBusMixHandle BusMixHandle = FBusMixHandle::Get(*BusMixPtr.Get(), RefProxies.BusMixes);
+				if (BusMixHandle.IsValid())
+				{
+					BusMixHandle.FindProxy().SetEnabled();
+				}
+				else
+				{
+					BusMixHandle = FBusMixHandle::Create(*BusMixPtr.Get(), RefProxies.BusMixes);
+				}
+
 				ManuallyActivatedBusMixes.Add(MoveTemp(BusMixHandle));
 			}
 		});
@@ -321,8 +343,12 @@ namespace AudioModulation
 			if (BusMixPtr.IsValid())
 			{
 				const USoundControlBusMix& BusMix = *BusMixPtr.Get();
-				FBusMixHandle BusMixHandle = FBusMixHandle::Create(BusMix, RefProxies.BusMixes);
-				ManuallyActivatedBusMixes.Remove(MoveTemp(BusMixHandle));
+				FBusMixHandle BusMixHandle = FBusMixHandle::Get(BusMix, RefProxies.BusMixes);
+				if (BusMixHandle.IsValid())
+				{
+					FModulatorBusMixProxy& MixProxy = BusMixHandle.FindProxy();
+					MixProxy.SetStopping();
+				}
 			}
 		});
 	}
@@ -423,21 +449,52 @@ namespace AudioModulation
 
 		// Update mix values and apply to prescribed buses.
 		// Track bus mixes ready to remove
-		TArray<FBusMixId> MixesToDeactivate;
+		TSet<FBusMixId> StoppedMixIds;
 		for (TPair<FBusMixId, FModulatorBusMixProxy>& Pair : RefProxies.BusMixes)
 		{
+			const FModulatorBusMixProxy::EStatus LastStatus = Pair.Value.GetStatus();
 			Pair.Value.Update(Elapsed, RefProxies.Buses);
-			if (Pair.Value.CanDestroy())
+			const FModulatorBusMixProxy::EStatus CurrentStatus = Pair.Value.GetStatus();
+
+			switch (CurrentStatus)
 			{
-				UE_LOG(LogAudioModulation, Log, TEXT("Audio modulation mix '%s' stopped."), *Pair.Value.GetName());
-				MixesToDeactivate.Add(Pair.Key);
+				case FModulatorBusMixProxy::EStatus::Enabled:
+				case FModulatorBusMixProxy::EStatus::Stopping:
+				break;
+
+				case FModulatorBusMixProxy::EStatus::Stopped:
+				{
+					if (LastStatus != CurrentStatus)
+					{
+						UE_LOG(LogAudioModulation, Log, TEXT("Audio modulation mix '%s' stopped."), *Pair.Value.GetName());
+					}
+					StoppedMixIds.Add(Pair.Key);
+				}
+				break;
+
+				default:
+				{
+					checkf(false, TEXT("Invalid or unsupported BusMix EStatus state advancement."));
+				}
+				break;
 			}
 		}
 
-		// Destroy mixes that have stopped
-		for (const FBusMixId& MixId : MixesToDeactivate)
+		// Destroy mixes that have stopped (must be done outside mix update
+		// loop above to avoid destroying while iterating, which can occur
+		// when update moves bus mix from 'stopping' status to 'stopped')
+		for (const FBusMixId& MixId : StoppedMixIds)
 		{
-			RefProxies.BusMixes.Remove(MixId);
+			FBusMixHandle MixHandle = FBusMixHandle::Get(MixId, RefProxies.BusMixes);
+
+			// Expected to be valid given the fact that the proxy is available in the prior loop
+			check(MixHandle.IsValid());
+
+			// Expected to only have two references (one for transient 'MixHandle' and one in
+			// ManuallyActivated set). Nothing else should be keeping mixes active.
+			check(MixHandle.FindProxy().GetRefCount() == 2);
+
+			ManuallyActivatedBusMixes.Remove(MoveTemp(MixHandle));
 		}
 
 		SET_DWORD_STAT(STAT_AudioModulationBusCount, RefProxies.Buses.Num());
@@ -449,6 +506,86 @@ namespace AudioModulation
 #endif // !UE_BUILD_SHIPPING
 	}
 
+	void FAudioModulationImpl::SaveMixToProfile(const USoundControlBusMix& InBusMix, const int32 InProfileIndex)
+	{
+		check(IsInGameThread());
+
+		const TWeakObjectPtr<const USoundControlBusMix> MixToSerialize = &InBusMix;
+		RunCommandOnAudioThread([this, MixToSerialize, InProfileIndex]()
+		{
+			if (!MixToSerialize.IsValid())
+			{
+				return;
+			}
+
+			const FBusMixId MixId = static_cast<FBusMixId>(MixToSerialize->GetUniqueID());
+			const FString   MixName = MixToSerialize->GetName();
+			const FString   MixPath = MixToSerialize->GetPathName();
+
+			FBusMixHandle MixHandle = FBusMixHandle::Get(MixId, RefProxies.BusMixes);
+			const bool bIsActive = MixHandle.IsValid();
+			if (!MixHandle.IsValid())
+			{
+				UE_LOG(LogAudioModulation, Display, TEXT("Mix '%s' is inactive, saving default object to profile '%i'."), *MixName, InProfileIndex);
+				AsyncTask(ENamedThreads::GameThread, [this, MixToSerialize, InProfileIndex]()
+				{
+					AudioModulation::FProfileSerializer::Serialize(*MixToSerialize.Get(), InProfileIndex);
+				});
+				return;
+			}
+
+			UE_LOG(LogAudioModulation, Display, TEXT("Mix '%s' is active, saving current mix proxy state to profile '%i'."), *MixName, InProfileIndex);
+			AudioModulation::FModulatorBusMixProxy& MixProxy = MixHandle.FindProxy();
+			TMap<FBusId, FSoundModulationValue> PassedChannelInfo;
+			for (TPair<FBusId, FModulatorBusMixChannelProxy>& Pair : MixProxy.Channels)
+			{
+				FModulatorBusMixChannelProxy& Channel = Pair.Value;
+				PassedChannelInfo.Add(Pair.Key, Channel.Value);
+			}
+
+			AsyncTask(ENamedThreads::GameThread, [this, PassedChannelInfo, MixId, MixToSerialize, InProfileIndex]()
+			{
+				if (!MixToSerialize.IsValid())
+				{
+					return;
+				}
+						
+				TMap<FBusId, FSoundModulationValue> ChannelInfo = PassedChannelInfo;
+				USoundControlBusMix* TempMix = NewObject<USoundControlBusMix>(GetTransientPackage(), *FGuid().ToString(EGuidFormats::Short));
+
+				// Buses on proxy may differ than those on uobject definition, so iterate and find by cached ids
+				// and add to temp mix to be serialized.
+				for (TObjectIterator<USoundControlBusBase> Itr; Itr; ++Itr)
+				{
+					if (USoundControlBusBase* Bus = *Itr)
+					{
+						FBusId ItrBusId = static_cast<FBusId>(Bus->GetUniqueID());
+						if (FSoundModulationValue* Value = ChannelInfo.Find(ItrBusId))
+						{
+							FSoundControlBusMixChannel BusMixChannel;
+							BusMixChannel.Bus = Bus;
+							BusMixChannel.Value = *Value;
+							TempMix->Channels.Add(MoveTemp(BusMixChannel));
+						}
+					}
+				}
+
+				const FString MixPath = MixToSerialize->GetPathName();
+				AudioModulation::FProfileSerializer::Serialize(*TempMix, InProfileIndex, &MixPath);
+			});
+		});
+	}
+
+	void FAudioModulationImpl::LoadMixFromProfile(const int32 InProfileIndex, USoundControlBusMix& OutBusMix)
+	{
+		const FString MixPath = OutBusMix.GetPathName();
+
+		USoundControlBusMix* TempMix = NewObject<USoundControlBusMix>(GetTransientPackage(), *FGuid().ToString(EGuidFormats::Short));
+		AudioModulation::FProfileSerializer::Deserialize(InProfileIndex, *TempMix, &MixPath);
+
+		UpdateMix(TempMix->Channels, OutBusMix);
+	}
+
 	void FAudioModulationImpl::RunCommandOnAudioThread(TFunction<void()> Cmd)
 	{
 		if (IsInAudioThread())
@@ -461,9 +598,37 @@ namespace AudioModulation
 		}
 	}
 
-	void FAudioModulationImpl::UpdateMix(const USoundControlBusMix& InMix, const TArray<FSoundControlBusMixChannel>& InChannels)
+	void FAudioModulationImpl::UpdateMix(const TArray<FSoundControlBusMixChannel>& InChannels, USoundControlBusMix& InOutMix, bool bInUpdateObject)
 	{
-		const FBusMixId MixId = static_cast<FBusMixId>(InMix.GetUniqueID());
+		if (bInUpdateObject)
+		{
+			TMap<uint32, const FSoundControlBusMixChannel*> UpdatedChannelBusses;
+			for (const FSoundControlBusMixChannel& Channel : InChannels)
+			{
+				if (Channel.Bus)
+				{
+					UpdatedChannelBusses.Add(Channel.Bus->GetUniqueID(), &Channel);
+				}
+			}
+
+			bool bMarkDirty = false;
+			for (FSoundControlBusMixChannel& Channel : InOutMix.Channels)
+			{
+				if (!Channel.Bus)
+				{
+					continue;
+				}
+
+				if (const FSoundControlBusMixChannel* BusChannel = UpdatedChannelBusses.FindRef(Channel.Bus->GetUniqueID()))
+				{
+					Channel = *BusChannel;
+					bMarkDirty = true;
+				}
+			}
+			InOutMix.MarkPackageDirty();
+		}
+
+		const FBusMixId MixId = static_cast<FBusMixId>(InOutMix.GetUniqueID());
 		RunCommandOnAudioThread([this, MixId, InChannels]()
 		{
 			if (FModulatorBusMixProxy* BusMixes = RefProxies.BusMixes.Find(MixId))
@@ -474,15 +639,57 @@ namespace AudioModulation
 	}
 
 	void FAudioModulationImpl::UpdateMixByFilter(
-		const USoundControlBusMix&					InMix,
-		const FString&								InAddressFilter,
-		const TSubclassOf<USoundControlBusBase>&	InBusClass,
-		const FSoundModulationValue&				InValue)
+		const FString&							 InAddressFilter,
+		const TSubclassOf<USoundControlBusBase>& InBusClass,
+		const FSoundModulationValue&			 InValue,
+		USoundControlBusMix&					 InOutMix,
+		bool									 bInUpdateObject)
 	{
-		const FString	AddressFilter	= InAddressFilter;
-		const uint32	ClassId			= InBusClass ? InBusClass->GetUniqueID() : USoundControlBusBase::StaticClass()->GetUniqueID();
-		const FBusMixId MixId			= static_cast<FBusMixId>(InMix.GetUniqueID());
+		const uint32 ClassId = InBusClass ? InBusClass->GetUniqueID() : USoundControlBusBase::StaticClass()->GetUniqueID();
+
+		if (bInUpdateObject)
+		{
+			bool bMarkDirty = false;
+			static const uint32 BaseBusClassId = USoundControlBusBase::StaticClass()->GetUniqueID();
+			for (FSoundControlBusMixChannel& Channel : InOutMix.Channels)
+			{
+				if (!Channel.Bus)
+				{
+					continue;
+				}
+
+				if (ClassId != BaseBusClassId && Channel.Bus->GetClass()->GetUniqueID() != ClassId)
+				{
+					continue;
+				}
+
+				if (!FAudioAddressPattern::PartsMatch(InAddressFilter, Channel.Bus->Address))
+				{
+					continue;
+				}
+
+				Channel.Value.TargetValue = InValue.TargetValue;
+
+				if (InValue.AttackTime >= 0.0f)
+				{
+					Channel.Value.AttackTime = InValue.AttackTime;
+				}
+
+				if (InValue.ReleaseTime >= 0.0f)
+				{
+					Channel.Value.ReleaseTime = InValue.ReleaseTime;
+				}
+				bMarkDirty = true;
+			}
+
+			if (bMarkDirty)
+			{
+				InOutMix.MarkPackageDirty();
+			}
+		}
 
+		const FString	AddressFilter = InAddressFilter;
+		const FBusMixId MixId = static_cast<FBusMixId>(InOutMix.GetUniqueID());
 		RunCommandOnAudioThread([this, ClassId, MixId, AddressFilter, InValue]()
 		{
 			if (FModulatorBusMixProxy* MixProxy = RefProxies.BusMixes.Find(MixId))
diff --git a/Engine/Plugins/Runtime/AudioModulation/Source/AudioModulation/Private/AudioModulationInternal.h b/Engine/Plugins/Runtime/AudioModulation/Source/AudioModulation/Private/AudioModulationInternal.h
--- a/Engine/Plugins/Runtime/AudioModulation/Source/AudioModulation/Private/AudioModulationInternal.h
+++ b/Engine/Plugins/Runtime/AudioModulation/Source/AudioModulation/Private/AudioModulationInternal.h
@@ -57,8 +57,23 @@ namespace AudioModulation
 		bool ProcessControls(const uint32 InSourceId, FSoundModulationControls& OutControls);
 		void ProcessModulators(const float Elapsed);
 
-		void UpdateMix(const USoundControlBusMix& InMix, const TArray<FSoundControlBusMixChannel>& InChannels);
-		void UpdateMixByFilter(const USoundControlBusMix& InMix, const FString& InAddressFilter, const TSubclassOf<USoundControlBusBase>& InClassFilter, const FSoundModulationValue& InValue);
+		/* Saves mix to .ini profile for fast iterative development that does not require re-cooking a mix */
+		void SaveMixToProfile(const USoundControlBusMix& InBusMix, const int32 InProfileIndex);
+
+		/* Loads mix from .ini profile for iterative development that does not require re-cooking a mix */
+		void LoadMixFromProfile(const int32 InProfileIndex, USoundControlBusMix& OutBusMix);
+
+		/*
+		 * Updates mix/mix by filter, modifying the mix instance if it is active. If bInUpdateObject is true,
+		 * updates UObject definition in addition to proxy.
+		 */
+		void UpdateMix(const TArray<FSoundControlBusMixChannel>& InChannels, USoundControlBusMix& InOutMix, bool bInUpdateObject = false);
+		void UpdateMixByFilter(const FString& InAddressFilter, const TSubclassOf<USoundControlBusBase>& InClassFilter, const FSoundModulationValue& InValue, USoundControlBusMix& InOutMix, bool bInUpdateObject = false);
+
+		/*
+		 * Commits any changes from a modulator type applied to a UObject definition
+		 * to modulator instance if active (i.e. Control Bus, Control Bus Mix, Control Bus Modulator)
+		 */
 		void UpdateModulator(const USoundModulatorBase& InModulator);
 
 	private:
@@ -90,6 +105,8 @@ namespace AudioModulation
 
 	private:
 		FAudioModulationDebugger Debugger;
+
+		friend FModulatorBusMixChannelProxy;
 #endif // !UE_BUILD_SHIPPING
 	};
 } // namespace AudioModulation
@@ -129,11 +146,14 @@ namespace AudioModulation
 		void DeactivateBusMix(const USoundControlBusBase& InBusMix) { }
 		void DeactivateLFO(const USoundBusModulatorLFO& InLFO) { }
 
+		void SaveMixToProfile(const USoundControlBusMix& InBusMix, const int32 InProfileIndex) { }
+		void LoadMixFromProfile(const int32 InProfileIndex, USoundControlBusMix& OutBusMix) { }
+
 		bool ProcessControls(const uint32 InSourceId, FSoundModulationControls& OutControls) { return false; }
 		void ProcessModulators(const float Elapsed) { }
 
-		void UpdateMix(const USoundControlBusMix& InMix, const TArray<FSoundControlBusMixChannel>& InChannels) { }
-		void UpdateMixByFilter(const USoundControlBusMix& InMix, const FString& InAddressFilter, const TSubclassOf<USoundControlBusBase>& InClassFilter, const FSoundModulationValue& InValue) { }
+		void UpdateMix(const TArray<FSoundControlBusMixChannel>& InChannels, USoundControlBusMix& InOutMix, bool bUpdateObject = false) { }
+		void UpdateMixByFilter(const FString& InAddressFilter, const TSubclassOf<USoundControlBusBase>& InClassFilter, const FSoundModulationValue& InValue, USoundControlBusMix& InOutMix, bool bUpdateObject = false) { }
 		void UpdateModulator(const USoundModulatorBase& InModulator) { }
 	};
 } // namespace AudioModulation
diff --git a/Engine/Plugins/Runtime/AudioModulation/Source/AudioModulation/Private/AudioModulationProfileSerializer.h b/Engine/Plugins/Runtime/AudioModulation/Source/AudioModulation/Private/AudioModulationProfileSerializer.h
new file mode 100644
--- /dev/null
+++ b/Engine/Plugins/Runtime/AudioModulation/Source/AudioModulation/Private/AudioModulationProfileSerializer.h
@@ -0,0 +1,218 @@
+// Copyright Epic Games, Inc. All Rights Reserved.
+#pragma once
+
+#include "CoreMinimal.h"
+#include "GenericPlatform/GenericPlatformFile.h"
+#include "HAL/PlatformFilemanager.h"
+#include "Logging/LogMacros.h"
+#include "Misc/ConfigCacheIni.h"
+#include "SoundControlBusMix.h"
+#include "UObject/Object.h"
+#include "UObject/ObjectMacros.h"
+
+
+namespace AudioModulation
+{
+	class FProfileSerializer
+	{
+		static bool ChannelValueToFloat(const FConfigSection& InSection, const FName& InMember, float& OutValue)
+		{
+			if (const FConfigValue* Value = InSection.Find(InMember))
+			{
+				OutValue = FCString::Atof(*Value->GetValue());
+				return true;
+			}
+
+			return false;
+		}
+
+		static FString GetConfigDir()
+		{
+			return FPaths::GeneratedConfigDir() + TEXT("AudioModulation");
+		}
+
+		static bool GetConfigDirectory(const FString& InDir, bool bCreateDirIfMissing)
+		{
+			IPlatformFile& PlatformFile = FPlatformFileManager::Get().GetPlatformFile();
+			if (!bCreateDirIfMissing && PlatformFile.DirectoryExists(*InDir))
+			{
+				return false;
+			}
+
+			TArray<FString> DirArray;
+			InDir.ParseIntoArray(DirArray, TEXT("/"));
+
+			FString SubPath = GetConfigDir();
+
+			if (!PlatformFile.CreateDirectory(*SubPath))
+			{
+				return false;
+			}
+
+			for (FString& Dir : DirArray)
+			{
+				SubPath /= Dir;
+				if (!PlatformFile.CreateDirectory(*SubPath))
+				{
+					return false;
+				}
+			}
+
+			return true;
+		}
+
+		static bool GetMixProfileConfigPath(const FString& Name, const FString& Path, const int32 InProfileIndex, bool bInCreateDirIfMissing, FString& OutMixProfileConfig)
+		{
+			FString ConfigDir = FPaths::GetPath(Path);
+			FString ConfigPath = GetConfigDir() + ConfigDir / Name;
+
+			ConfigPath += InProfileIndex > 0
+				? FString::Printf(TEXT("%i.ini"), InProfileIndex)
+				: FString(TEXT(".ini"));
+
+			if (!GConfig)
+			{
+				UE_LOG(LogAudioModulation, Error,
+					TEXT("Failed to find mix profile '%s' config file. Config cache not initialized."),
+					*Path, *ConfigPath);
+				return false;
+			}
+
+			if (!GetConfigDirectory(ConfigDir, bInCreateDirIfMissing))
+			{
+				UE_LOG(LogAudioModulation, Error, TEXT("Failed to %s mix profile '%s' config directory."),
+					bInCreateDirIfMissing ? TEXT("create") : TEXT("find"),
+					*Path,
+					*ConfigDir);
+				return false;
+			}
+
+			OutMixProfileConfig = ConfigPath;
+			return true;
+		}
+
+	public:
+		static bool Serialize(const USoundControlBusMix& InBusMix, const int32 InProfileIndex, const FString* InMixPathOverride = nullptr)
+		{
+			check(IsInGameThread());
+
+			const FString Name = InMixPathOverride ? FPaths::GetBaseFilename(*InMixPathOverride) : InBusMix.GetName();
+			const FString Path = InMixPathOverride ? *InMixPathOverride : InBusMix.GetPathName();
+
+			FString ConfigPath;
+			if (!GetMixProfileConfigPath(Name, Path, InProfileIndex, true /* bCreateDir */, ConfigPath))
+			{
+				return false;
+			}
+
+			FConfigFile& ConfigFile = GConfig->FindOrAdd(ConfigPath);
+			ConfigFile.Reset();
+
+			for (const FSoundControlBusMixChannel& Channel : InBusMix.Channels)
+			{
+				if (Channel.Bus)
+				{
+					FConfigSection& ConfigSection = ConfigFile.Add(Channel.Bus->GetPathName());
+					ConfigSection.Add(GET_MEMBER_NAME_CHECKED(FSoundModulationValue, AttackTime), FConfigValue(FString::Printf(TEXT("%f"), Channel.Value.AttackTime)));
+					ConfigSection.Add(GET_MEMBER_NAME_CHECKED(FSoundModulationValue, ReleaseTime), FConfigValue(FString::Printf(TEXT("%f"), Channel.Value.ReleaseTime)));
+					ConfigSection.Add(GET_MEMBER_NAME_CHECKED(FSoundModulationValue, TargetValue), FConfigValue(FString::Printf(TEXT("%f"), Channel.Value.TargetValue)));
+				}
+			}
+			ConfigFile.Dirty = true;
+
+			UE_LOG(LogAudioModulation, Display, TEXT("Serialized mix to '%s'"), *Path);
+			GConfig->Flush(false, ConfigPath);
+
+			return true;
+		}
+
+		static bool Deserialize(int32 InProfileIndex, USoundControlBusMix& InOutBusMix, const FString* InMixPathOverride = nullptr)
+		{
+			check(IsInGameThread());
+
+			const FString Name = InMixPathOverride ? FPaths::GetBaseFilename(*InMixPathOverride) : InOutBusMix.GetName();
+			const FString Path = InMixPathOverride ? *InMixPathOverride : InOutBusMix.GetPathName();
+
+			FString ConfigPath;
+			if (!GetMixProfileConfigPath(Name, Path, InProfileIndex, false /* bCreateDir */, ConfigPath))
+			{
+				return false;
+			}
+
+			FConfigFile* ConfigFile = GConfig->Find(ConfigPath, false /* CreateIfNotFound */);
+			if (!ConfigFile)
+			{
+				UE_LOG(LogAudioModulation, Warning, TEXT("Config '%s' not found. Mix '%s' not loaded from config profile."), *ConfigPath , *Name);
+				return false;
+			}
+
+			bool bMarkDirty = false;
+			TSet<FString> SectionsProcessed;
+			TArray<FSoundControlBusMixChannel>& Channels = InOutBusMix.Channels;
+			for (int32 i = Channels.Num() - 1; i >= 0; --i)
+			{
+				FSoundControlBusMixChannel& Channel = Channels[i];
+				if (Channel.Bus)
+				{
+					FString PathName = Channel.Bus->GetPathName();
+					if (FConfigSection* ConfigSection = ConfigFile->Find(PathName))
+					{
+						ChannelValueToFloat(*ConfigSection, GET_MEMBER_NAME_CHECKED(FSoundModulationValue, AttackTime), Channel.Value.AttackTime);
+						ChannelValueToFloat(*ConfigSection, GET_MEMBER_NAME_CHECKED(FSoundModulationValue, ReleaseTime), Channel.Value.ReleaseTime);
+						ChannelValueToFloat(*ConfigSection, GET_MEMBER_NAME_CHECKED(FSoundModulationValue, TargetValue), Channel.Value.TargetValue);
+						SectionsProcessed.Emplace(MoveTemp(PathName));
+						bMarkDirty = true;
+					}
+					else
+					{
+						Channels.RemoveAt(i);
+					}
+				}
+				else
+				{
+					Channels.RemoveAt(i);
+				}
+			}
+
+			TArray<FString> ConfigSectionNames;
+			ConfigFile->GetKeys(ConfigSectionNames);
+			for (FString& SectionName : ConfigSectionNames)
+			{
+				if (!SectionsProcessed.Contains(SectionName))
+				{
+					const FConfigSection& ConfigSection = ConfigFile->FindChecked(SectionName);
+					UObject* BusObj = FSoftObjectPath(SectionName).TryLoad();
+					if (USoundControlBusBase* Bus = Cast<USoundControlBusBase>(BusObj))
+					{
+						FSoundControlBusMixChannel Channel;
+						Channel.Bus = Bus;
+
+						ChannelValueToFloat(ConfigSection, GET_MEMBER_NAME_CHECKED(FSoundModulationValue, AttackTime),  Channel.Value.AttackTime);
+						ChannelValueToFloat(ConfigSection, GET_MEMBER_NAME_CHECKED(FSoundModulationValue, ReleaseTime), Channel.Value.ReleaseTime);
+						ChannelValueToFloat(ConfigSection, GET_MEMBER_NAME_CHECKED(FSoundModulationValue, TargetValue), Channel.Value.TargetValue);
+
+						Channels.Emplace(MoveTemp(Channel));
+
+						bMarkDirty = true;
+					}
+					else
+					{
+						UE_LOG(LogAudioModulation, Warning,
+							TEXT("Bus missing or invalid type at '%s'. Profile '%s' channel not added/updated for mix '%s'"),
+							*SectionName,
+							*ConfigPath,
+							*Path);
+					}
+				}
+			}
+
+			if (bMarkDirty)
+			{
+				InOutBusMix.MarkPackageDirty();
+			}
+
+			UE_LOG(LogAudioModulation, Display, TEXT("Deserialized mix from '%s'"), *ConfigPath);
+			return true;
+		}
+	};
+} // namespace AudioModulation
\ No newline at end of file
diff --git a/Engine/Plugins/Runtime/AudioModulation/Source/AudioModulation/Private/AudioModulationStatics.cpp b/Engine/Plugins/Runtime/AudioModulation/Source/AudioModulation/Private/AudioModulationStatics.cpp
--- a/Engine/Plugins/Runtime/AudioModulation/Source/AudioModulation/Private/AudioModulationStatics.cpp
+++ b/Engine/Plugins/Runtime/AudioModulation/Source/AudioModulation/Private/AudioModulationStatics.cpp
@@ -1,22 +1,113 @@
 // Copyright Epic Games, Inc. All Rights Reserved.
 #include "AudioModulationStatics.h"
 
+#include "Async/Async.h"
 #include "AudioDevice.h"
 #include "AudioModulation.h"
 #include "AudioModulationInternal.h"
 #include "AudioModulationLogging.h"
-
+#include "AudioModulationProfileSerializer.h"
+#include "CoreGlobals.h"
 #include "Engine/Engine.h"
+#include "HAL/IConsoleManager.h"
+#include "HAL/PlatformFilemanager.h"
+#include "Misc/ConfigCacheIni.h"
+#include "Misc/Paths.h"
 #include "SoundControlBus.h"
 #include "SoundControlBusMix.h"
 
 #define LOCTEXT_NAMESPACE "AudioModulationStatics"
 
 
-namespace
+static FAutoConsoleCommand GModulationSaveMixProfile(
+	TEXT("au.Modulation.SaveMixProfile"),
+	TEXT("Saves modulation mix profile to the config save directory.\n"
+		"Path - Path to Object\n"
+		"ProfileIndex - (Optional) Index of profile (defaults to 0)"),
+	FConsoleCommandWithArgsDelegate::CreateStatic(
+		[](const TArray<FString>& Args)
+		{
+			if (Args.Num() < 1)
+			{
+				UE_LOG(LogAudioModulation, Error, TEXT("Failed to save mix profile: Path not provided"));
+				return;
+			}
+
+			const FString& Path = Args[0];
+			int32 ProfileIndex = 0;
+			if (Args.Num() > 1)
+			{
+				ProfileIndex = FCString::Atoi(*Args[1]);
+			}
+
+			FSoftObjectPath ObjPath = Path;
+			if (UObject* MixObj = ObjPath.TryLoad())
+			{
+				if (USoundControlBusMix* Mix = Cast<USoundControlBusMix>(MixObj))
+				{
+					UAudioModulationStatics::SaveMixToProfile(Mix, Mix, ProfileIndex);
+					return;
+				}
+			}
+
+			UE_LOG(LogAudioModulation, Error, TEXT("Failed to save mix '%s' to profile index '%i'"), *Path, ProfileIndex);
+		}
+	)
+);
+
+static FAutoConsoleCommand GModulationLoadMixProfile(
+	TEXT("au.Modulation.LoadMixProfile"),
+	TEXT("Loads modulation mix profile from the config save directory.\n"
+		"Path - Path to Object to load\n"
+		"Activate - (Optional) Whether or not to activate/update the mix once it is loaded (default: true)."
+		"ProfileIndex - (Optional) Index of profile (default: 0)"),
+	FConsoleCommandWithArgsDelegate::CreateStatic(
+		[](const TArray<FString>& Args)
+		{
+			if (Args.Num() < 1)
+			{
+				UE_LOG(LogAudioModulation, Error, TEXT("Failed to load mix profile: Object path not provided"));
+				return;
+			}
+
+			const FString& Path = Args[0];
+			int32 ProfileIndex = 0;
+			if (Args.Num() > 1)
+			{
+				ProfileIndex = FCString::Atoi(*Args[1]);
+			}
+
+			bool bActivateUpdate = true;
+			if (Args.Num() > 2)
+			{
+				bActivateUpdate = FCString::ToBool(*Args[2]);
+			}
+
+			FSoftObjectPath ObjPath = Path;
+			if (UObject* MixObj = ObjPath.TryLoad())
+			{
+				if (USoundControlBusMix* Mix = Cast<USoundControlBusMix>(MixObj))
+				{
+					UAudioModulationStatics::LoadMixFromProfile(Mix, Mix, bActivateUpdate, ProfileIndex);
+
+					if (bActivateUpdate)
+					{
+						UAudioModulationStatics::UpdateModulator(Mix, Mix);
+					}
+					return;
+				}
+			}
+
+			UE_LOG(LogAudioModulation, Error, TEXT("Failed to load mix '%s' from profile index '%i'"), *Path, ProfileIndex);
+		}
+	)
+);
+
+
+namespace AudioModulation
 {
 	template <class T>
-	T* CreateModulatorBus(const UObject* WorldContextObject, FName Name, float DefaultValue, bool Activate)
+	T* CreateBus(const UObject* WorldContextObject, FName Name, float DefaultValue, bool Activate)
 	{
 		UWorld* World = UAudioModulationStatics::GetAudioWorld(WorldContextObject);
 		if (!World)
@@ -30,7 +121,7 @@ namespace
 
 		if (Activate)
 		{
-			if (AudioModulation::FAudioModulationImpl* ModulationImpl = UAudioModulationStatics::GetModulationImpl(World))
+			if (FAudioModulationImpl* ModulationImpl = UAudioModulationStatics::GetModulationImpl(World))
 			{
 				ModulationImpl->ActivateBus(*NewBus);
 			}
@@ -38,10 +129,8 @@ namespace
 
 		return NewBus;
 	}
-} // namespace <>
+} // namespace AudioModulation
 
-//////////////////////////////////////////////////////////////////////////
-// UAudioModulationStatics
 
 UAudioModulationStatics::UAudioModulationStatics(const FObjectInitializer& ObjectInitializer)
 	: Super(ObjectInitializer)
@@ -130,22 +219,22 @@ AudioModulation::FAudioModulationImpl* UAudioModulationStatics::GetModulationImp
 
 USoundVolumeControlBus* UAudioModulationStatics::CreateVolumeBus(const UObject* WorldContextObject, FName Name, float DefaultValue, bool Activate)
 {
-	return CreateModulatorBus<USoundVolumeControlBus>(WorldContextObject, Name, DefaultValue, Activate);
+	return AudioModulation::CreateBus<USoundVolumeControlBus>(WorldContextObject, Name, DefaultValue, Activate);
 }
 
 USoundPitchControlBus* UAudioModulationStatics::CreatePitchBus(const UObject* WorldContextObject, FName Name, float DefaultValue, bool Activate)
 {
-	return CreateModulatorBus<USoundPitchControlBus>(WorldContextObject, Name, DefaultValue, Activate);
+	return AudioModulation::CreateBus<USoundPitchControlBus>(WorldContextObject, Name, DefaultValue, Activate);
 }
 
 USoundLPFControlBus* UAudioModulationStatics::CreateLPFBus(const UObject* WorldContextObject, FName Name, float DefaultValue, bool Activate)
 {
-	return CreateModulatorBus<USoundLPFControlBus>(WorldContextObject, Name, DefaultValue, Activate);
+	return AudioModulation::CreateBus<USoundLPFControlBus>(WorldContextObject, Name, DefaultValue, Activate);
 }
 
 USoundHPFControlBus* UAudioModulationStatics::CreateHPFBus(const UObject* WorldContextObject, FName Name, float DefaultValue, bool Activate)
 {
-	return CreateModulatorBus<USoundHPFControlBus>(WorldContextObject, Name, DefaultValue, Activate);
+	return AudioModulation::CreateBus<USoundHPFControlBus>(WorldContextObject, Name, DefaultValue, Activate);
 }
 
 USoundBusModulatorLFO* UAudioModulationStatics::CreateLFO(const UObject* WorldContextObject, FName Name, float Amplitude, float Frequency, float Offset, bool Activate)
@@ -250,6 +339,34 @@ void UAudioModulationStatics::DeactivateBusModulator(const UObject* WorldContext
 	}
 }
 
+void UAudioModulationStatics::SaveMixToProfile(const UObject* WorldContextObject, USoundControlBusMix* BusMix, int32 ProfileIndex)
+{
+	UWorld* World = GetAudioWorld(WorldContextObject);
+	if (AudioModulation::FAudioModulationImpl* ModulationImpl = GetModulationImpl(World))
+	{
+		if (BusMix)
+		{
+			return ModulationImpl->SaveMixToProfile(*BusMix, ProfileIndex);
+		}
+	}
+}
+
+void UAudioModulationStatics::LoadMixFromProfile(const UObject* WorldContextObject, USoundControlBusMix* BusMix, bool bActivate, int32 ProfileIndex)
+{
+	UWorld* World = GetAudioWorld(WorldContextObject);
+	if (AudioModulation::FAudioModulationImpl* ModulationImpl = GetModulationImpl(World))
+	{
+		if (BusMix)
+		{
+			ModulationImpl->LoadMixFromProfile(ProfileIndex, *BusMix);
+			if (bActivate)
+			{
+				ActivateBusMix(WorldContextObject, BusMix);
+			}
+		}
+	}
+}
+
 void UAudioModulationStatics::UpdateMix(const UObject* WorldContextObject, USoundControlBusMix* Mix, TArray<FSoundControlBusMixChannel> Channels)
 {
 	if (Mix)
@@ -257,7 +374,9 @@ void UAudioModulationStatics::UpdateMix(const UObject* WorldContextObject, USoun
 		UWorld* World = GetAudioWorld(WorldContextObject);
 		if (AudioModulation::FAudioModulationImpl* ModulationImpl = GetModulationImpl(World))
 		{
-			ModulationImpl->UpdateMix(*Mix, Channels);
+			// UObject representation is not updated in this form of the call as doing so from
+			// PIE can result in an unstable state where UObject is modified but not properly dirtied.
+			ModulationImpl->UpdateMix(Channels, *Mix, false /* bUpdateObject */);
 		}
 	}
 }
@@ -277,7 +396,10 @@ void UAudioModulationStatics::UpdateMixByFilter(
 		if (AudioModulation::FAudioModulationImpl* ModulationImpl = GetModulationImpl(World))
 		{
 			FSoundModulationValue ModValue(Value, AttackTime, ReleaseTime);
-			ModulationImpl->UpdateMixByFilter(*Mix, AddressFilter, BusClassFilter, ModValue);
+
+			// UObject representation is not updated in this form of the call as doing so from
+			// PIE can result in an unstable state where UObject is modified but not properly dirtied.
+			ModulationImpl->UpdateMixByFilter(AddressFilter, BusClassFilter, ModValue, *Mix, false /* bUpdateObject */);
 		}
 	}
 }
diff --git a/Engine/Plugins/Runtime/AudioModulation/Source/AudioModulation/Private/AudioModulationStatics.h b/Engine/Plugins/Runtime/AudioModulation/Source/AudioModulation/Private/AudioModulationStatics.h
--- a/Engine/Plugins/Runtime/AudioModulation/Source/AudioModulation/Private/AudioModulationStatics.h
+++ b/Engine/Plugins/Runtime/AudioModulation/Source/AudioModulation/Private/AudioModulationStatics.h
@@ -40,19 +40,28 @@ public:
 	/** Activates a bus. Does nothing if an instance of the provided bus is already active
 	 * @param Bus - Bus to activate
 	 */
-	UFUNCTION(BlueprintCallable, Category = "Audio", DisplayName = "Activate Control Bus", meta = (WorldContext = "WorldContextObject", Keywords = "activate modulation modulator control bus"))
+	UFUNCTION(BlueprintCallable, Category = "Audio", DisplayName = "Activate Control Bus", meta = (
+		WorldContext = "WorldContextObject", 
+		Keywords = "activate modulation modulator control bus")
+	)
 	static void ActivateBus(const UObject* WorldContextObject, USoundControlBusBase* Bus);
 
 	/** Activates a bus modulator mix. Does nothing if an instance of the provided bus mix is already active
 	 * @param BusMix - Mix to activate
 	 */
-	UFUNCTION(BlueprintCallable, Category = "Audio", DisplayName = "Activate Control Bus Mix", meta = (WorldContext = "WorldContextObject", Keywords = "activate modulation modulator control bus mix"))
+	UFUNCTION(BlueprintCallable, Category = "Audio", DisplayName = "Activate Control Bus Mix", meta = (
+		WorldContext = "WorldContextObject", 
+		Keywords = "activate modulation modulator control bus mix")
+	)
 	static void ActivateBusMix(const UObject* WorldContextObject, USoundControlBusMix* BusMix);
 
 	/** Activates a bus modulator (eg. LFO). Does nothing if an instance of the provided modulator is already active
 	 * @param Modulator - Modulator to activate
 	 */
-	UFUNCTION(BlueprintCallable, Category = "Audio", DisplayName = "Activate Control Bus Modulator", meta = (WorldContext = "WorldContextObject", Keywords = "activate modulation modulator lfo"))
+	UFUNCTION(BlueprintCallable, Category = "Audio", DisplayName = "Activate Control Bus Modulator", meta = (
+		WorldContext = "WorldContextObject", 
+		Keywords = "activate modulation modulator lfo")
+	)
 	static void ActivateBusModulator(const UObject* WorldContextObject, USoundBusModulatorBase* Modulator);
 
 	/** Creates a volume modulation bus with the provided default value.
@@ -60,7 +69,11 @@ public:
 	 * @param DefaultValue - Default value for created bus
 	 * @param Activate - Whether or not to activate bus on creation.
 	 */
-	UFUNCTION(BlueprintCallable, Category = "Audio", DisplayName = "Create Volume Control Bus", meta = (WorldContext = "WorldContextObject", Keywords = "make create bus modulation volume modulator"))
+	UFUNCTION(BlueprintCallable, Category = "Audio", DisplayName = "Create Volume Control Bus", meta = (
+		AdvancedDisplay = "3",
+		WorldContext = "WorldContextObject", 
+		Keywords = "make create bus modulation volume modulator")
+	)
 	static USoundVolumeControlBus* CreateVolumeBus(const UObject* WorldContextObject, FName Name, float DefaultValue, bool Activate);
 
 	/** Creates a pitch modulation bus with the provided default value.
@@ -68,24 +81,36 @@ public:
 	 * @param DefaultValue - Default value for created bus
 	 * @param Activate - Whether or not to activate bus on creation.
 	 */
-	UFUNCTION(BlueprintCallable, Category = "Audio", DisplayName = "Create Pitch Control Bus", meta = (WorldContext = "WorldContextObject", Keywords = "make create bus modulation pitch modulator"))
-	static USoundPitchControlBus* CreatePitchBus(const UObject* WorldContextObject, FName Name, float DefaultValue, bool Activate);
+	UFUNCTION(BlueprintCallable, Category = "Audio", DisplayName = "Create Pitch Control Bus", meta = (
+		AdvancedDisplay = "3",
+		WorldContext = "WorldContextObject",
+		Keywords = "make create bus modulation pitch modulator")
+	)
+	static USoundPitchControlBus* CreatePitchBus(const UObject* WorldContextObject, FName Name, float DefaultValue, bool Activate = true);
 
 	/** Creates a high-pass filter (HPF) modulation bus with the provided default value.
 	 * @param Name - Name of bus
 	 * @param DefaultValue - Default value for created bus
 	 * @param Activate - Whether or not to activate bus on creation.
 	 */
-	UFUNCTION(BlueprintCallable, Category = "Audio", DisplayName = "Create HPF Control Bus", meta = (WorldContext = "WorldContextObject", Keywords = "make create bus modulation HPF modulator"))
-	static USoundHPFControlBus* CreateHPFBus(const UObject* WorldContextObject, FName Name, float DefaultValue, bool Activate);
+	UFUNCTION(BlueprintCallable, Category = "Audio", DisplayName = "Create HPF Control Bus", meta = (
+		AdvancedDisplay = "3", 
+		WorldContext = "WorldContextObject", 
+		Keywords = "make create bus modulation HPF modulator")
+	)
+	static USoundHPFControlBus* CreateHPFBus(const UObject* WorldContextObject, FName Name, float DefaultValue, bool Activate = true);
 
 	/** Creates a low-pass filter (LPF) modulation bus with the provided default value.
 	 * @param Name - Name of bus
 	 * @param DefaultValue - Default value for created bus
 	 * @param Activate - Whether or not to activate bus on creation.
 	 */
-	UFUNCTION(BlueprintCallable, Category = "Audio", DisplayName = "Create LPF Control Bus", meta = (WorldContext = "WorldContextObject", Keywords = "make create bus modulation LPF modulator"))
-	static USoundLPFControlBus* CreateLPFBus(const UObject* WorldContextObject, FName Name, float DefaultValue, bool Activate);
+	UFUNCTION(BlueprintCallable, Category = "Audio", DisplayName = "Create LPF Control Bus", meta = (
+		AdvancedDisplay = "3",
+		WorldContext = "WorldContextObject",
+		Keywords = "make create bus modulation LPF modulator")
+	)
+	static USoundLPFControlBus* CreateLPFBus(const UObject* WorldContextObject, FName Name, float DefaultValue, bool Activate = true);
 
 	/** Creates an LFO modulator.
 	 * @param Name - Name of LFO
@@ -94,60 +119,131 @@ public:
 	 * @param Offset - Offset of new LFO.
 	 * @param Activate - Whether or not to activate lfo on creation.
 	 */
-	UFUNCTION(BlueprintCallable, Category = "Audio", DisplayName = "Create Control Bus LFO", meta = (WorldContext = "WorldContextObject", Keywords = "make create lfo modulation modulator"))
-	static USoundBusModulatorLFO* CreateLFO(const UObject* WorldContextObject, FName Name, float Amplitude, float Frequency, float Offset, bool Activate);
+	UFUNCTION(BlueprintCallable, Category = "Audio", DisplayName = "Create Control Bus LFO", meta = (
+		AdvancedDisplay = "5",
+		WorldContext = "WorldContextObject",
+		Keywords = "make create lfo modulation modulator")
+	)
+	static USoundBusModulatorLFO* CreateLFO(
+		const UObject* WorldContextObject, 
+		FName Name, 
+		float Amplitude, 
+		float Frequency, 
+		float Offset,
+		bool Activate = true
+	);
 
 	/** Creates a channel used to mix a control bus.
 	 * @param Bus - Bus channel is in charge of applying mix value to.
 	 * @param Channels - Value for added bus channel to target when mix is active.
 	 * @param Attack/ReleaseTime - Time in seconds for channel to mix in/out.
 	 */
-	UFUNCTION(BlueprintCallable, Category = "Audio", DisplayName = "Create Control Bus Mix Channel", meta = (AdvancedDisplay = "3", WorldContext = "WorldContextObject", Keywords = "make create control bus mix modulation modulator channel"))
-	static FSoundControlBusMixChannel CreateBusMixChannel(const UObject* WorldContextObject, USoundControlBusBase* Bus, float Value, float AttackTime = 0.1f, float ReleaseTime = 0.1f);
+	UFUNCTION(BlueprintCallable, Category = "Audio", DisplayName = "Create Control Bus Mix Channel", meta = (
+		AdvancedDisplay = "3",
+		WorldContext = "WorldContextObject",
+		Keywords = "make create control bus mix modulation modulator channel")
+	)
+	static FSoundControlBusMixChannel CreateBusMixChannel(
+		const UObject* WorldContextObject,
+		USoundControlBusBase* Bus,
+		float Value,
+		float AttackTime = 0.1f,
+		float ReleaseTime = 0.1f);
 
 	/** Creates a modulation bus mix and adds a bus channel set to the provided target value
 	 * @param Name - Name of mix.
 	 * @param Channels - Channels mix is responsible for.
 	 * @param Activate - Whether or not to activate mix on creation.
 	 */
-	UFUNCTION(BlueprintCallable, Category = "Audio", DisplayName = "Create Control Bus Mix", meta = (WorldContext = "WorldContextObject", Keywords = "make create control bus mix modulation modulator"))
-	static USoundControlBusMix* CreateBusMix(const UObject* WorldContextObject, FName Name, TArray<FSoundControlBusMixChannel> Channels, bool Activate);
+	UFUNCTION(BlueprintCallable, Category = "Audio", DisplayName = "Create Control Bus Mix", meta = (
+		WorldContext = "WorldContextObject", 
+		Keywords = "make create control bus mix modulation modulator")
+	)
+	static USoundControlBusMix* CreateBusMix(
+		const UObject* WorldContextObject,
+		FName Name, 
+		TArray<FSoundControlBusMixChannel> Channels, 
+		bool Activate);
 
 	/** Deactivates a bus. Does nothing if an instance of the provided bus is already inactive
 	 * @param Bus - Scope of modulator
 	 */
-	UFUNCTION(BlueprintCallable, Category = "Audio", DisplayName = "Deactivate Control Bus", meta = (WorldContext = "WorldContextObject", Keywords = "deactivate modulation modulator bus"))
+	UFUNCTION(BlueprintCallable, Category = "Audio", DisplayName = "Deactivate Control Bus", meta = (
+		WorldContext = "WorldContextObject", 
+		Keywords = "deactivate modulation modulator bus")
+	)
 	static void DeactivateBus(const UObject* WorldContextObject, USoundControlBusBase* Bus);
 
 	/** Deactivates a modulation bus mix. Does nothing if an instance of the provided bus mix is already inactive
 	 * @param BusMix - Mix to deactivate
 	 */
-	UFUNCTION(BlueprintCallable, Category = "Audio", DisplayName = "Deactivate Control Bus Mix", meta = (WorldContext = "WorldContextObject", Keywords = "deactivate modulation modulator"))
+	UFUNCTION(BlueprintCallable, Category = "Audio", DisplayName = "Deactivate Control Bus Mix", meta = (
+		WorldContext = "WorldContextObject",
+		Keywords = "deactivate modulation modulator")
+	)
 	static void DeactivateBusMix(const UObject* WorldContextObject, USoundControlBusMix* BusMix);
 
 	/** Deactivates a bus modulator. Does nothing if an instance of the provided bus mix is already inactive
 	 * @param Modulator - Modulator to activate
 	 * @param Scope - Scope of modulator
 	 */
-	UFUNCTION(BlueprintCallable, Category = "Audio", DisplayName = "Deactivate Control Bus Modulator", meta = (WorldContext = "WorldContextObject", Keywords = "deactivate bus modulation modulator"))
+	UFUNCTION(BlueprintCallable, Category = "Audio", DisplayName = "Deactivate Control Bus Modulator", meta = (
+		WorldContext = "WorldContextObject",
+		Keywords = "deactivate bus modulation modulator")
+)
 	static void DeactivateBusModulator(const UObject* WorldContextObject, USoundBusModulatorBase* Modulator);
 
-	/** Updates a mix with the provided channel data.
+	/** Saves control bus mix to a profile, serialized to an ini file.  If mix is loaded, uses current proxy's state. 
+	 * If not, uses default UObject representation.
+	 * @param BusMix - Mix object to serialize to profile .ini.
+	 * @param ProfileIndex - Index of profile, allowing multiple profiles can be saved for single mix object. If 0, saves to default ini profile (no suffix).
+	 */
+	UFUNCTION(BlueprintCallable, Category = "Audio", DisplayName = "Save Control Bus Mix to Profile", meta = (
+		WorldContext = "WorldContextObject",
+		AdvancedDisplay = "2",
+		Keywords = "save serialize bus control modulation mix modulator ini")
+	)
+	static void SaveMixToProfile(const UObject* WorldContextObject, USoundControlBusMix* BusMix, int32 ProfileIndex = 0);
+
+	/** Loads control bus mix from a profile into UObject mix definition, deserialized from an ini file.
+	 * @param BusMix - Mix object to deserialize profile .ini to.
+	 * @param bActivate - If true, activate mix upon loading from profile.
+	 * @param ProfileIndex - Index of profile, allowing multiple profiles to be loaded to single mix object. If <= 0, loads from default profile (no suffix).
+	 */
+	UFUNCTION(BlueprintCallable, Category = "Audio", DisplayName = "Load Control Bus Mix From Profile", meta = (
+		WorldContext = "WorldContextObject",
+		AdvancedDisplay = "2",
+		Keywords = "load deserialize control bus modulation mix modulator ini")
+	)
+	static void LoadMixFromProfile(const UObject* WorldContextObject, USoundControlBusMix* BusMix, bool bActivate = true, int32 ProfileIndex = 0);
+
+	/** Sets a mix with the provided channel data if channels provided in active instance proxy of mix. Does not update UObject definition of mix.
 	 * @param Mix - Mix to update
 	 * @param Channels - Channels to set.  If channel's bus is not referenced by mix, channel's update request is ignored.
+	 * @param bUpdateObject - If true, will dirty mix object and update channels on the SoundControlBusMix object in addition
+	 * to updating the audio thread proxy.
 	 */
-	UFUNCTION(BlueprintCallable, Category = "Audio", DisplayName = "Update Control Bus Mix", meta = (WorldContext = "WorldContextObject", Keywords = "update bus control modulation modulator mix channel"))
+	UFUNCTION(BlueprintCallable, Category = "Audio", DisplayName = "Set Control Bus Mix", meta = (
+		WorldContext = "WorldContextObject",
+		Keywords = "set bus control modulation modulator mix channel")
+	)
 	static void UpdateMix(const UObject* WorldContextObject, USoundControlBusMix* Mix, TArray<FSoundControlBusMixChannel> Channels);
 
-	/** Updates filtered channels of a given class to a provided target value.
+	/** Sets filtered channels of a given class to a provided target value for active instance of mix. Does not update UObject definition of mix.
 	 * @param Mix - Mix to modify
 	 * @param AddressFilter - Address filter to apply to provided mix's channels.
 	 * @param BusClass - Filters buses by subclass.
 	 * @param Value - Target value to mix filtered channels to.
 	 * @param AttackTime - If non-negative, updates the attack time for the resulting bus channels found matching the provided filter.
 	 * @param ReleaseTime - If non-negative, updates the release time for the resulting bus channels found matching the provided filter.
+	 * @param bUpdateObject - If true, will dirty mix object and update channels on the SoundControlBusMix object in addition to updating 
+	 * the audio thread proxy.
 	 */
-	UFUNCTION(BlueprintCallable, Category = "Audio", DisplayName = "Update Control Bus Mix By Filter", meta = (AdvancedDisplay = "5", WorldContext = "WorldContextObject", Keywords = "update bus control class modulation modulator mix channel value filter"))
+	UFUNCTION(BlueprintCallable, Category = "Audio", DisplayName = "Set Control Bus Mix By Filter", meta = (
+		AdvancedDisplay = "5",
+		WorldContext = "WorldContextObject",
+		Keywords = "set bus control class modulation modulator mix channel value filter")
+	)
 	static void UpdateMixByFilter(
 		const UObject*						WorldContextObject,
 		USoundControlBusMix*				Mix,
@@ -157,9 +253,13 @@ public:
 		float								AttackTime	= -1.0f,
 		float								ReleaseTime = -1.0f);
 
-	/** Commits updates from a UObject definition of a modulator (e.g. Bus, Bus Mix, LFO) to active instance (does not apply if modulator has not been activated).
+	/** Commits updates from a UObject definition of a modulator (e.g. Bus, Bus Mix, LFO) to active instance in audio thread
+	 * (ignored if modulator type has not been activated).
 	 * @param Modulator - Modulator to update
 	 */
-	UFUNCTION(BlueprintCallable, Category = "Audio", DisplayName = "Update Modulator", meta = (WorldContext = "WorldContextObject", Keywords = "update set control bus mix modulation modulator"))
+	UFUNCTION(BlueprintCallable, Category = "Audio", DisplayName = "Update Modulator", meta = (
+		WorldContext = "WorldContextObject",
+		Keywords = "update set control bus mix modulation modulator")
+	)
 	static void UpdateModulator(const UObject* WorldContextObject, USoundModulatorBase* Modulator);
 };
diff --git a/Engine/Plugins/Runtime/AudioModulation/Source/AudioModulation/Private/SoundControlBusMix.cpp b/Engine/Plugins/Runtime/AudioModulation/Source/AudioModulation/Private/SoundControlBusMix.cpp
--- a/Engine/Plugins/Runtime/AudioModulation/Source/AudioModulation/Private/SoundControlBusMix.cpp
+++ b/Engine/Plugins/Runtime/AudioModulation/Source/AudioModulation/Private/SoundControlBusMix.cpp
@@ -6,8 +6,17 @@
 #include "AudioModulation.h"
 #include "AudioModulationInternal.h"
 #include "AudioModulationLogging.h"
+#include "AudioModulationProfileSerializer.h"
+#include "AudioModulationStatics.h"
 #include "Engine/World.h"
 
+#if WITH_EDITOR
+#include "Framework/Notifications/NotificationManager.h"
+#include "Widgets/Notifications/SNotificationList.h"
+#endif // WITH_EDITOR
+
+
+#define LOCTEXT_NAMESPACE "AudioModulation"
 
 FSoundControlBusMixChannel::FSoundControlBusMixChannel()
 	: Bus(nullptr)
@@ -22,6 +31,9 @@ FSoundControlBusMixChannel::FSoundControlBusMixChannel(USoundControlBusBase* InB
 
 USoundControlBusMix::USoundControlBusMix(const FObjectInitializer& ObjectInitializer)
 	: Super(ObjectInitializer)
+#if WITH_EDITORONLY_DATA
+	, ProfileIndex(0)
+#endif // WITH_EDITORONLY_DATA
 {
 }
 
@@ -44,6 +56,42 @@ void USoundControlBusMix::BeginDestroy()
 	}
 }
 
+#if WITH_EDITOR
+void USoundControlBusMix::LoadMixFromProfile()
+{
+	if (AudioModulation::FProfileSerializer::Deserialize(ProfileIndex, *this))
+	{
+		FNotificationInfo Info(FText::Format(
+			LOCTEXT("SoundControlBusMix_LoadSucceeded", "'Control Bus Mix '{0}' profile {1} loaded successfully."),
+			FText::FromName(GetFName()),
+			FText::AsNumber(ProfileIndex)
+		));
+		Info.bFireAndForget = true;
+		Info.ExpireDuration = 2.0f;
+		Info.bUseThrobber = true;
+		FSlateNotificationManager::Get().AddNotification(Info);
+	}
+}
+
+void USoundControlBusMix::SaveMixToProfile()
+{
+	if (AudioModulation::FProfileSerializer::Serialize(*this, ProfileIndex))
+	{
+		{
+			FNotificationInfo Info(FText::Format(
+				LOCTEXT("SoundControlBusMix_SaveSucceeded", "'Control Bus Mix '{0}' profile {1} saved successfully."),
+					FText::FromName(GetFName()),
+					FText::AsNumber(ProfileIndex)
+			));
+			Info.bFireAndForget = true;
+			Info.ExpireDuration = 2.0f;
+			Info.bUseThrobber = true;
+			FSlateNotificationManager::Get().AddNotification(Info);
+		}
+	}
+}
+#endif // WITH_EDITOR
+
 namespace AudioModulation
 {
 	FModulatorBusMixChannelProxy::FModulatorBusMixChannelProxy(const FSoundControlBusMixChannel& Channel)
@@ -52,6 +100,10 @@ namespace AudioModulation
 		, ClassId(Channel.Bus->GetClass()->GetUniqueID())
 		, Value(Channel.Value)
 	{
+		check(Channel.Bus);
+		AudioModulation::FAudioModulationImpl* ModulationImpl = UAudioModulationStatics::GetModulationImpl(Channel.Bus->GetWorld());
+		check(ModulationImpl);
+		BusHandle = FBusHandle::Create(*Channel.Bus, ModulationImpl->RefProxies.Buses);
 	}
 
 	FModulatorBusMixProxy::FModulatorBusMixProxy(const USoundControlBusMix& InBusMix)
@@ -67,36 +119,41 @@ namespace AudioModulation
 		return *this;
 	}
 
-	bool FModulatorBusMixProxy::CanDestroy() const
+	FModulatorBusMixProxy::EStatus FModulatorBusMixProxy::GetStatus() const
 	{
-		return Status == BusMixStatus::Stopped;
+		return Status;
+	}
+
+	void FModulatorBusMixProxy::Reset()
+	{
+		Channels.Reset();
 	}
 
 	void FModulatorBusMixProxy::Init(const USoundControlBusMix& InBusMix)
 	{
+		FChannelMap CachedChannels = Channels;
 		Channels.Reset();
 
-		Status = BusMixStatus::Enabled;
+		Status = EStatus::Enabled;
 		for (const FSoundControlBusMixChannel& Channel : InBusMix.Channels)
 		{
 			if (Channel.Bus)
 			{
 				auto BusId = static_cast<const AudioModulation::FBusId>(Channel.Bus->GetUniqueID());
-#if !UE_BUILD_SHIPPING
-				if (Channels.Contains(BusId))
+				
+				FModulatorBusMixChannelProxy ChannelProxy(Channel);
+				if (const FModulatorBusMixChannelProxy* CachedChannel = CachedChannels.Find(BusId))
 				{
-					UE_LOG(LogAudioModulation, Warning,
-						TEXT("USoundControlBusMix '%s' already contains bus '%s'. Only one representative channel for this bus added."),
-						*InBusMix.GetFullName(), *Channel.Bus->GetFullName());
+					ChannelProxy.Value.SetCurrentValue(CachedChannel->Value.GetCurrentValue());
 				}
-#endif // UE_BUILD_SHIPPING
-				Channels.Emplace(BusId, FModulatorBusMixChannelProxy(Channel));
+				
+				Channels.Emplace(BusId, ChannelProxy);
 			}
 			else
 			{
 				UE_LOG(LogAudioModulation, Warning,
 					TEXT("USoundControlBusMix '%s' has channel with no bus specified. "
-						"Mix activated but channel ignored."),
+						"Mix instance initialized but channel ignored."),
 					*InBusMix.GetFullName());
 			}
 		}
@@ -152,14 +209,14 @@ namespace AudioModulation
 
 	void FModulatorBusMixProxy::SetEnabled()
 	{
-		Status = BusMixStatus::Enabled;
+		Status = EStatus::Enabled;
 	}
 
 	void FModulatorBusMixProxy::SetStopping()
 	{
-		if (Status == BusMixStatus::Enabled)
+		if (Status == EStatus::Enabled)
 		{
-			Status = BusMixStatus::Stopping;
+			Status = EStatus::Stopping;
 		}
 	}
 
@@ -176,7 +233,7 @@ namespace AudioModulation
 				MixChannelValue.Update(Elapsed);
 
 				const float CurrentValue = MixChannelValue.GetCurrentValue();
-				if (Status == BusMixStatus::Stopping)
+				if (Status == EStatus::Stopping)
 				{
 					MixChannelValue.TargetValue = BusProxy->GetDefaultValue();
 					if (!FMath::IsNearlyEqual(MixChannelValue.TargetValue, CurrentValue))
@@ -194,7 +251,9 @@ namespace AudioModulation
 
 		if (bRequestStop)
 		{
-			Status = BusMixStatus::Stopped;
+			Status = EStatus::Stopped;
 		}
 	}
 } // namespace AudioModulation
+
+#undef LOCTEXT_NAMESPACE // AudioModulation
diff --git a/Engine/Plugins/Runtime/AudioModulation/Source/AudioModulation/Private/SoundModulationPatch.cpp b/Engine/Plugins/Runtime/AudioModulation/Source/AudioModulation/Private/SoundModulationPatch.cpp
--- a/Engine/Plugins/Runtime/AudioModulation/Source/AudioModulation/Private/SoundModulationPatch.cpp
+++ b/Engine/Plugins/Runtime/AudioModulation/Source/AudioModulation/Private/SoundModulationPatch.cpp
@@ -317,13 +317,5 @@ namespace AudioModulation
 		{
 			Controls.Add(Patch.Control, FModulationPatchProxy(Patch, OutRefProxies));
 		}
-
-		for (const USoundControlBusMix* Mix : Settings.Mixes)
-		{
-			if (Mix)
-			{
-				Mixes.Add(FBusMixHandle::Create(*Mix, OutRefProxies.BusMixes));
-			}
-		}
 	}
 } // namespace AudioModulation
\ No newline at end of file
diff --git a/Engine/Plugins/Runtime/AudioModulation/Source/AudioModulation/Private/SoundModulationProxy.h b/Engine/Plugins/Runtime/AudioModulation/Source/AudioModulation/Private/SoundModulationProxy.h
--- a/Engine/Plugins/Runtime/AudioModulation/Source/AudioModulation/Private/SoundModulationProxy.h
+++ b/Engine/Plugins/Runtime/AudioModulation/Source/AudioModulation/Private/SoundModulationProxy.h
@@ -120,6 +120,22 @@ namespace AudioModulation
 			}
 		}
 
+		/*
+		 * Creates a handle to a proxy modulation object tracked in the provided InProxyMap if it exists, otherwise returns invalid handle.
+		 */
+		static TProxyHandle<IdType, ProxyType, ProxyUObjType> Get(const IdType ObjectId, TMap<IdType, ProxyType>& InProxyMap)
+		{
+			if (ProxyType* Proxy = InProxyMap.Find(ObjectId))
+			{
+				TProxyHandle<IdType, ProxyType, ProxyUObjType> NewHandle(ObjectId, InProxyMap);
+				Proxy->IncRef();
+
+				return MoveTemp(NewHandle);
+			} 
+
+			return TProxyHandle<IdType, ProxyType, ProxyUObjType>();
+		}
+
 		/*
 		 * Creates a handle to a proxy modulation object tracked in the provided InProxyMap if it exists, otherwise returns invalid handle.
 		 */
@@ -131,7 +147,7 @@ namespace AudioModulation
 				TProxyHandle<IdType, ProxyType, ProxyUObjType> NewHandle(ObjectId, InProxyMap);
 				Proxy->IncRef();
 
-				return NewHandle;
+				return MoveTemp(NewHandle);
 			} 
 
 			return TProxyHandle<IdType, ProxyType, ProxyUObjType>();
@@ -247,6 +263,8 @@ namespace AudioModulation
 		{
 		}
 
+		virtual ~TModulatorProxyBase() = default;
+
 		IdType GetId() const
 		{
 			return Id;
@@ -264,14 +282,6 @@ namespace AudioModulation
 		}
 	};
 
-	struct FModulatorBusMixChannelProxy : public TModulatorProxyBase<FBusId>
-	{
-		FModulatorBusMixChannelProxy(const FSoundControlBusMixChannel& Channel);
-		FString Address;
-		uint32 ClassId;
-		FSoundModulationValue Value;
-	};
-
 	class FModulatorLFOProxy;
 
 	using FLFOProxyMap = TMap<FLFOId, FModulatorLFOProxy>;
@@ -412,10 +422,19 @@ namespace AudioModulation
 	using FBusProxyMap = TMap<FBusId, FControlBusProxy>;
 	using FBusHandle = TProxyHandle<FBusId, FControlBusProxy, USoundControlBusBase>;
 
+	struct FModulatorBusMixChannelProxy : public TModulatorProxyBase<FBusId>
+	{
+		FModulatorBusMixChannelProxy(const FSoundControlBusMixChannel& Channel);
+		FString Address;
+		uint32 ClassId;
+		FSoundModulationValue Value;
+		FBusHandle BusHandle;
+	};
+
 	class FModulatorBusMixProxy : public TModulatorProxyRefType<FBusMixId, FModulatorBusMixProxy, USoundControlBusMix>
 	{
 	public:
-		enum class BusMixStatus : uint8
+		enum class EStatus : uint8
 		{
 			Enabled,
 			Stopping,
@@ -426,7 +445,10 @@ namespace AudioModulation
 
 		FModulatorBusMixProxy& operator =(const USoundControlBusMix& InBusMix);
 
-		bool CanDestroy() const;
+		EStatus GetStatus() const;
+
+		// Resets channel map
+		void Reset();
 
 		void SetEnabled();
 		void SetMix(const TArray<FSoundControlBusMixChannel>& InChannels);
@@ -435,11 +457,12 @@ namespace AudioModulation
 
 		void Update(const float Elapsed, FBusProxyMap& ProxyMap);
 
-		TMap<FBusId, FModulatorBusMixChannelProxy> Channels;
+		using FChannelMap = TMap<FBusId, FModulatorBusMixChannelProxy>;
+		FChannelMap Channels;
 
 	private:
 		void Init(const USoundControlBusMix& InBusMix);
-		BusMixStatus Status;
+		EStatus Status;
 	};
 
 	using FBusMixProxyMap = TMap<FBusMixId, FModulatorBusMixProxy>;
@@ -504,8 +527,6 @@ namespace AudioModulation
 		FModulationPatchProxy Highpass;
 
 		TMap<FName, FModulationPatchProxy> Controls;
-
-		TArray<FBusMixHandle> Mixes;
 	};
 
 	struct FReferencedProxies
diff --git a/Engine/Plugins/Runtime/AudioModulation/Source/AudioModulation/Private/SoundModulationValue.cpp b/Engine/Plugins/Runtime/AudioModulation/Source/AudioModulation/Private/SoundModulationValue.cpp
--- a/Engine/Plugins/Runtime/AudioModulation/Source/AudioModulation/Private/SoundModulationValue.cpp
+++ b/Engine/Plugins/Runtime/AudioModulation/Source/AudioModulation/Private/SoundModulationValue.cpp
@@ -27,6 +27,11 @@ FSoundModulationValue::FSoundModulationValue(float InValue, float InAttackTime,
 {
 }
 
+void FSoundModulationValue::SetCurrentValue(float InValue)
+{
+	Value = InValue;
+}
+
 float FSoundModulationValue::GetCurrentValue() const
 {
 	// Current does not require update to be called if
diff --git a/Engine/Plugins/Runtime/AudioModulation/Source/AudioModulation/Public/SoundControlBusMix.h b/Engine/Plugins/Runtime/AudioModulation/Source/AudioModulation/Public/SoundControlBusMix.h
--- a/Engine/Plugins/Runtime/AudioModulation/Source/AudioModulation/Public/SoundControlBusMix.h
+++ b/Engine/Plugins/Runtime/AudioModulation/Source/AudioModulation/Public/SoundControlBusMix.h
@@ -25,7 +25,7 @@ struct FSoundControlBusMixChannel
 	FSoundControlBusMixChannel(USoundControlBusBase* InBus, const float TargetValue);
 
 	/* Bus controlled by channel. */
-	UPROPERTY(EditAnywhere, Category = Channel, BlueprintReadOnly)
+	UPROPERTY(EditAnywhere, Category = Channel, BlueprintReadWrite)
 	USoundControlBusBase* Bus;
 
 	/* Value mix is set to. */
@@ -38,7 +38,20 @@ class USoundControlBusMix : public USoundModulatorBase
 {
 	GENERATED_UCLASS_BODY()
 
+protected:
+#if WITH_EDITOR
+	UFUNCTION(Category = Mix, meta = (CallInEditor = "true"))
+	void LoadMixFromProfile();
+
+	UFUNCTION(Category = Mix, meta = (CallInEditor = "true"))
+	void SaveMixToProfile();
+#endif // WITH_EDITOR
+
 public:
+#if WITH_EDITORONLY_DATA
+	UPROPERTY(EditAnywhere, Transient, Category = Mix)
+	uint32 ProfileIndex;
+#endif // WITH_EDITORONLY_DATA
 
 	virtual void BeginDestroy() override;
 
diff --git a/Engine/Plugins/Runtime/AudioModulation/Source/AudioModulation/Public/SoundModulationPatch.h b/Engine/Plugins/Runtime/AudioModulation/Source/AudioModulation/Public/SoundModulationPatch.h
--- a/Engine/Plugins/Runtime/AudioModulation/Source/AudioModulation/Public/SoundModulationPatch.h
+++ b/Engine/Plugins/Runtime/AudioModulation/Source/AudioModulation/Public/SoundModulationPatch.h
@@ -314,13 +314,6 @@ public:
 	UPROPERTY()
 	TArray<FSoundControlModulationPatch> Controls;
 
-	// Mixes that will applied and removed when sounds utilizing settings
-	// play and stop respectively. If mix has already been applied manually,
-	// mix will be removed once all sound settings referencing mix stop. Manual
-	// mix activation is ignored if already activated by means of modulation settings.
-	UPROPERTY(EditAnywhere, BlueprintReadWrite, Category = Mixes)
-	TArray<USoundControlBusMix*> Mixes;
-
 #if WITH_EDITOR
 	AUDIOMODULATION_API void OnPostEditChange(UWorld* World);
 
diff --git a/Engine/Plugins/Runtime/AudioModulation/Source/AudioModulation/Public/SoundModulationValue.h b/Engine/Plugins/Runtime/AudioModulation/Source/AudioModulation/Public/SoundModulationValue.h
--- a/Engine/Plugins/Runtime/AudioModulation/Source/AudioModulation/Public/SoundModulationValue.h
+++ b/Engine/Plugins/Runtime/AudioModulation/Source/AudioModulation/Public/SoundModulationValue.h
@@ -8,7 +8,7 @@
 
 
 USTRUCT(BlueprintType)
-struct FSoundModulationValue
+struct AUDIOMODULATION_API FSoundModulationValue
 {
 	GENERATED_USTRUCT_BODY()
 
@@ -20,13 +20,16 @@ struct FSoundModulationValue
 	float TargetValue;
 
 	/** Time it takes (in sec) to unitarily increase the bus value (from 0 to 1). */
-	UPROPERTY(EditAnywhere, Category = Modulation, BlueprintReadWrite, meta = (DisplayName = "Attack Time", ClampMin = "0.0", UIMin = "0.0"))
+	UPROPERTY(EditAnywhere, Category = Modulation, BlueprintReadWrite, meta = (DisplayName = "Attack Time (sec)", ClampMin = "0.0", UIMin = "0.0"))
 	float AttackTime;
 
 	/** Time it takes (in sec) to unitarily decrease the bus value (from 1 to 0). */
-	UPROPERTY(EditAnywhere, Category = Modulation, BlueprintReadWrite, meta = (DisplayName = "Release Time", ClampMin = "0.0", UIMin = "0.0"))
+	UPROPERTY(EditAnywhere, Category = Modulation, BlueprintReadWrite, meta = (DisplayName = "Release Time (sec)", ClampMin = "0.0", UIMin = "0.0"))
 	float ReleaseTime;
 
+	/** Set current value (for resetting value state only as circumvents lerp, and may result in discontinuity). */
+	void SetCurrentValue(float InValue);
+
 	/** Current value lerping toward target */
 	float GetCurrentValue() const;
 
