Missing from last CL

diff --git a/Engine/Plugins/Runtime/AudioModulation/Source/AudioModulation/AudioModulation.Build.cs b/Engine/Plugins/Runtime/AudioModulation/Source/AudioModulation/AudioModulation.Build.cs
--- a/Engine/Plugins/Runtime/AudioModulation/Source/AudioModulation/AudioModulation.Build.cs
+++ b/Engine/Plugins/Runtime/AudioModulation/Source/AudioModulation/AudioModulation.Build.cs
@@ -24,6 +24,17 @@ namespace UnrealBuildTool.Rules
 				}
 			);
 
+			if (Target.Type == TargetType.Editor)
+			{
+				PrivateDependencyModuleNames.AddRange(
+					new string[]
+					{
+						"Slate",
+						"SlateCore",
+					}
+				);
+			}
+
 			PublicDefinitions.Add("WITH_AUDIOMODULATION=1");
 		}
 	}
