- Expose editor style colors to BP for OSC controller to piggyback to public, non-editor API
- Move AudioModulationStatics.h to public
- Return channel value array when loading profile with new values deserialized from .ini

diff --git a/Engine/Plugins/Runtime/AudioModulation/Source/AudioModulation/Private/AudioModulationImpl.cpp b/Engine/Plugins/Runtime/AudioModulation/Source/AudioModulation/Private/AudioModulationImpl.cpp
--- a/Engine/Plugins/Runtime/AudioModulation/Source/AudioModulation/Private/AudioModulationImpl.cpp
+++ b/Engine/Plugins/Runtime/AudioModulation/Source/AudioModulation/Private/AudioModulationImpl.cpp
@@ -9,6 +9,7 @@
 #include "Engine/Engine.h"
 #include "IAudioExtensionPlugin.h"
 #include "Misc/CoreDelegates.h"
+#include "SoundControlBusMix.h"
 #include "SoundModulatorLFO.h"
 #include "SoundModulationPatch.h"
 #include "SoundModulationProxy.h"
@@ -576,14 +577,18 @@ namespace AudioModulation
 		});
 	}
 
-	void FAudioModulationImpl::LoadMixFromProfile(const int32 InProfileIndex, USoundControlBusMix& OutBusMix)
+	TArray<FSoundControlBusMixChannel> FAudioModulationImpl::LoadMixFromProfile(const int32 InProfileIndex, USoundControlBusMix& OutBusMix)
 	{
-		const FString MixPath = OutBusMix.GetPathName();
-
-		USoundControlBusMix* TempMix = NewObject<USoundControlBusMix>(GetTransientPackage(), *FGuid().ToString(EGuidFormats::Short));
-		AudioModulation::FProfileSerializer::Deserialize(InProfileIndex, *TempMix, &MixPath);
+		const FString TempName = FGuid::NewGuid().ToString(EGuidFormats::Short);
+		if (USoundControlBusMix* TempMix = NewObject<USoundControlBusMix>(GetTransientPackage(), *TempName))
+		{
+			const FString MixPath = OutBusMix.GetPathName();
+			AudioModulation::FProfileSerializer::Deserialize(InProfileIndex, *TempMix, &MixPath);
+			UpdateMix(TempMix->Channels, OutBusMix);
+			return TempMix->Channels;
+		}
 
-		UpdateMix(TempMix->Channels, OutBusMix);
+		return TArray<FSoundControlBusMixChannel>();
 	}
 
 	void FAudioModulationImpl::RunCommandOnAudioThread(TFunction<void()> Cmd)
diff --git a/Engine/Plugins/Runtime/AudioModulation/Source/AudioModulation/Private/AudioModulationInternal.h b/Engine/Plugins/Runtime/AudioModulation/Source/AudioModulation/Private/AudioModulationInternal.h
--- a/Engine/Plugins/Runtime/AudioModulation/Source/AudioModulation/Private/AudioModulationInternal.h
+++ b/Engine/Plugins/Runtime/AudioModulation/Source/AudioModulation/Private/AudioModulationInternal.h
@@ -60,8 +60,9 @@ namespace AudioModulation
 		/* Saves mix to .ini profile for fast iterative development that does not require re-cooking a mix */
 		void SaveMixToProfile(const USoundControlBusMix& InBusMix, const int32 InProfileIndex);
 
-		/* Loads mix from .ini profile for iterative development that does not require re-cooking a mix */
-		void LoadMixFromProfile(const int32 InProfileIndex, USoundControlBusMix& OutBusMix);
+		/* Loads mix from .ini profile for iterative development that does not require re-cooking a mix. Returns copy
+		 * of mix channel values saved in profile. */
+		TArray<FSoundControlBusMixChannel> LoadMixFromProfile(const int32 InProfileIndex, USoundControlBusMix& OutBusMix);
 
 		/*
 		 * Updates mix/mix by filter, modifying the mix instance if it is active. If bInUpdateObject is true,
diff --git a/Engine/Plugins/Runtime/AudioModulation/Source/AudioModulation/Private/AudioModulationStatics.cpp b/Engine/Plugins/Runtime/AudioModulation/Source/AudioModulation/Private/AudioModulationStatics.cpp
--- a/Engine/Plugins/Runtime/AudioModulation/Source/AudioModulation/Private/AudioModulationStatics.cpp
+++ b/Engine/Plugins/Runtime/AudioModulation/Source/AudioModulation/Private/AudioModulationStatics.cpp
@@ -351,20 +351,22 @@ void UAudioModulationStatics::SaveMixToProfile(const UObject* WorldContextObject
 	}
 }
 
-void UAudioModulationStatics::LoadMixFromProfile(const UObject* WorldContextObject, USoundControlBusMix* BusMix, bool bActivate, int32 ProfileIndex)
+TArray<FSoundControlBusMixChannel> UAudioModulationStatics::LoadMixFromProfile(const UObject* WorldContextObject, USoundControlBusMix* BusMix, bool bActivate, int32 ProfileIndex)
 {
-	UWorld* World = GetAudioWorld(WorldContextObject);
-	if (AudioModulation::FAudioModulationImpl* ModulationImpl = GetModulationImpl(World))
+	if (BusMix)
 	{
-		if (BusMix)
+		UWorld* World = GetAudioWorld(WorldContextObject);
+		if (AudioModulation::FAudioModulationImpl* ModulationImpl = GetModulationImpl(World))
 		{
-			ModulationImpl->LoadMixFromProfile(ProfileIndex, *BusMix);
 			if (bActivate)
 			{
 				ActivateBusMix(WorldContextObject, BusMix);
 			}
+			return ModulationImpl->LoadMixFromProfile(ProfileIndex, *BusMix);
 		}
 	}
+
+	return TArray<FSoundControlBusMixChannel>();
 }
 
 void UAudioModulationStatics::UpdateMix(const UObject* WorldContextObject, USoundControlBusMix* Mix, TArray<FSoundControlBusMixChannel> Channels)
diff --git a/Engine/Plugins/Runtime/AudioModulation/Source/AudioModulation/Private/AudioModulationStatics.h b/Engine/Plugins/Runtime/AudioModulation/Source/AudioModulation/Private/AudioModulationStatics.h
deleted file mode 100644
diff --git a/Engine/Plugins/Runtime/AudioModulation/Source/AudioModulation/Public/AudioModulationStatics.h b/Engine/Plugins/Runtime/AudioModulation/Source/AudioModulation/Public/AudioModulationStatics.h
new file mode 100644
--- /dev/null
+++ b/Engine/Plugins/Runtime/AudioModulation/Source/AudioModulation/Public/AudioModulationStatics.h
@@ -0,0 +1,266 @@
+// Copyright Epic Games, Inc. All Rights Reserved.
+#pragma once
+
+#include "CoreMinimal.h"
+#include "Kismet/BlueprintFunctionLibrary.h"
+#include "UObject/Object.h"
+#include "UObject/ObjectMacros.h"
+#include "Templates/SubclassOf.h"
+#include "SoundControlBusMix.h"
+#include "SoundControlBus.h"
+#include "AudioModulationStatics.generated.h"
+
+// Forward Declarations
+class USoundControlBus;
+class USoundControlBusMix;
+
+
+namespace AudioModulation
+{
+	class FAudioModulationImpl;
+} // namespace AudioModulation
+
+
+UCLASS()
+class UAudioModulationStatics : public UBlueprintFunctionLibrary
+{
+	GENERATED_UCLASS_BODY()
+
+public:
+	/**
+	 * Returns world associated with provided context object
+	 */
+	static UWorld* GetAudioWorld(const UObject* WorldContextObject);
+
+	/**
+	 * Returns modulation implementation associated with the provided world
+	 */
+	static AudioModulation::FAudioModulationImpl* GetModulationImpl(UWorld* World);
+
+	/** Activates a bus. Does nothing if an instance of the provided bus is already active
+	 * @param Bus - Bus to activate
+	 */
+	UFUNCTION(BlueprintCallable, Category = "Audio", DisplayName = "Activate Control Bus", meta = (
+		WorldContext = "WorldContextObject", 
+		Keywords = "activate modulation modulator control bus")
+	)
+	static void ActivateBus(const UObject* WorldContextObject, USoundControlBusBase* Bus);
+
+	/** Activates a bus modulator mix. Does nothing if an instance of the provided bus mix is already active
+	 * @param BusMix - Mix to activate
+	 */
+	UFUNCTION(BlueprintCallable, Category = "Audio", DisplayName = "Activate Control Bus Mix", meta = (
+		WorldContext = "WorldContextObject", 
+		Keywords = "activate modulation modulator control bus mix")
+	)
+	static void ActivateBusMix(const UObject* WorldContextObject, USoundControlBusMix* BusMix);
+
+	/** Activates a bus modulator (eg. LFO). Does nothing if an instance of the provided modulator is already active
+	 * @param Modulator - Modulator to activate
+	 */
+	UFUNCTION(BlueprintCallable, Category = "Audio", DisplayName = "Activate Control Bus Modulator", meta = (
+		WorldContext = "WorldContextObject", 
+		Keywords = "activate modulation modulator lfo")
+	)
+	static void ActivateBusModulator(const UObject* WorldContextObject, USoundBusModulatorBase* Modulator);
+
+	/** Creates a volume modulation bus with the provided default value.
+	 * @param Name - Name of bus
+	 * @param DefaultValue - Default value for created bus
+	 * @param Activate - Whether or not to activate bus on creation.
+	 */
+	UFUNCTION(BlueprintCallable, Category = "Audio", DisplayName = "Create Volume Control Bus", meta = (
+		AdvancedDisplay = "3",
+		WorldContext = "WorldContextObject", 
+		Keywords = "make create bus modulation volume modulator")
+	)
+	static USoundVolumeControlBus* CreateVolumeBus(const UObject* WorldContextObject, FName Name, float DefaultValue, bool Activate);
+
+	/** Creates a pitch modulation bus with the provided default value.
+	 * @param Name - Name of bus
+	 * @param DefaultValue - Default value for created bus
+	 * @param Activate - Whether or not to activate bus on creation.
+	 */
+	UFUNCTION(BlueprintCallable, Category = "Audio", DisplayName = "Create Pitch Control Bus", meta = (
+		AdvancedDisplay = "3",
+		WorldContext = "WorldContextObject",
+		Keywords = "make create bus modulation pitch modulator")
+	)
+	static USoundPitchControlBus* CreatePitchBus(const UObject* WorldContextObject, FName Name, float DefaultValue, bool Activate = true);
+
+	/** Creates a high-pass filter (HPF) modulation bus with the provided default value.
+	 * @param Name - Name of bus
+	 * @param DefaultValue - Default value for created bus
+	 * @param Activate - Whether or not to activate bus on creation.
+	 */
+	UFUNCTION(BlueprintCallable, Category = "Audio", DisplayName = "Create HPF Control Bus", meta = (
+		AdvancedDisplay = "3", 
+		WorldContext = "WorldContextObject", 
+		Keywords = "make create bus modulation HPF modulator")
+	)
+	static USoundHPFControlBus* CreateHPFBus(const UObject* WorldContextObject, FName Name, float DefaultValue, bool Activate = true);
+
+	/** Creates a low-pass filter (LPF) modulation bus with the provided default value.
+	 * @param Name - Name of bus
+	 * @param DefaultValue - Default value for created bus
+	 * @param Activate - Whether or not to activate bus on creation.
+	 */
+	UFUNCTION(BlueprintCallable, Category = "Audio", DisplayName = "Create LPF Control Bus", meta = (
+		AdvancedDisplay = "3",
+		WorldContext = "WorldContextObject",
+		Keywords = "make create bus modulation LPF modulator")
+	)
+	static USoundLPFControlBus* CreateLPFBus(const UObject* WorldContextObject, FName Name, float DefaultValue, bool Activate = true);
+
+	/** Creates an LFO modulator.
+	 * @param Name - Name of LFO
+	 * @param Amplitude - Amplitude of new LFO.
+	 * @param Frequency - Frequency of new LFO.
+	 * @param Offset - Offset of new LFO.
+	 * @param Activate - Whether or not to activate lfo on creation.
+	 */
+	UFUNCTION(BlueprintCallable, Category = "Audio", DisplayName = "Create Control Bus LFO", meta = (
+		AdvancedDisplay = "5",
+		WorldContext = "WorldContextObject",
+		Keywords = "make create lfo modulation modulator")
+	)
+	static USoundBusModulatorLFO* CreateLFO(
+		const UObject* WorldContextObject, 
+		FName Name, 
+		float Amplitude, 
+		float Frequency, 
+		float Offset,
+		bool Activate = true
+	);
+
+	/** Creates a channel used to mix a control bus.
+	 * @param Bus - Bus channel is in charge of applying mix value to.
+	 * @param Channels - Value for added bus channel to target when mix is active.
+	 * @param Attack/ReleaseTime - Time in seconds for channel to mix in/out.
+	 */
+	UFUNCTION(BlueprintCallable, Category = "Audio", DisplayName = "Create Control Bus Mix Channel", meta = (
+		AdvancedDisplay = "3",
+		WorldContext = "WorldContextObject",
+		Keywords = "make create control bus mix modulation modulator channel")
+	)
+	static FSoundControlBusMixChannel CreateBusMixChannel(
+		const UObject* WorldContextObject,
+		USoundControlBusBase* Bus,
+		float Value,
+		float AttackTime = 0.1f,
+		float ReleaseTime = 0.1f);
+
+	/** Creates a modulation bus mix and adds a bus channel set to the provided target value
+	 * @param Name - Name of mix.
+	 * @param Channels - Channels mix is responsible for.
+	 * @param Activate - Whether or not to activate mix on creation.
+	 */
+	UFUNCTION(BlueprintCallable, Category = "Audio", DisplayName = "Create Control Bus Mix", meta = (
+		WorldContext = "WorldContextObject", 
+		Keywords = "make create control bus mix modulation modulator")
+	)
+	static USoundControlBusMix* CreateBusMix(
+		const UObject* WorldContextObject,
+		FName Name, 
+		TArray<FSoundControlBusMixChannel> Channels, 
+		bool Activate);
+
+	/** Deactivates a bus. Does nothing if an instance of the provided bus is already inactive
+	 * @param Bus - Scope of modulator
+	 */
+	UFUNCTION(BlueprintCallable, Category = "Audio", DisplayName = "Deactivate Control Bus", meta = (
+		WorldContext = "WorldContextObject", 
+		Keywords = "deactivate modulation modulator bus")
+	)
+	static void DeactivateBus(const UObject* WorldContextObject, USoundControlBusBase* Bus);
+
+	/** Deactivates a modulation bus mix. Does nothing if an instance of the provided bus mix is already inactive
+	 * @param BusMix - Mix to deactivate
+	 */
+	UFUNCTION(BlueprintCallable, Category = "Audio", DisplayName = "Deactivate Control Bus Mix", meta = (
+		WorldContext = "WorldContextObject",
+		Keywords = "deactivate modulation modulator")
+	)
+	static void DeactivateBusMix(const UObject* WorldContextObject, USoundControlBusMix* BusMix);
+
+	/** Deactivates a bus modulator. Does nothing if an instance of the provided bus mix is already inactive
+	 * @param Modulator - Modulator to activate
+	 * @param Scope - Scope of modulator
+	 */
+	UFUNCTION(BlueprintCallable, Category = "Audio", DisplayName = "Deactivate Control Bus Modulator", meta = (
+		WorldContext = "WorldContextObject",
+		Keywords = "deactivate bus modulation modulator")
+)
+	static void DeactivateBusModulator(const UObject* WorldContextObject, USoundBusModulatorBase* Modulator);
+
+	/** Saves control bus mix to a profile, serialized to an ini file.  If mix is loaded, uses current proxy's state. 
+	 * If not, uses default UObject representation.
+	 * @param BusMix - Mix object to serialize to profile .ini.
+	 * @param ProfileIndex - Index of profile, allowing multiple profiles can be saved for single mix object. If 0, saves to default ini profile (no suffix).
+	 */
+	UFUNCTION(BlueprintCallable, Category = "Audio", DisplayName = "Save Control Bus Mix to Profile", meta = (
+		WorldContext = "WorldContextObject",
+		AdvancedDisplay = "2",
+		Keywords = "save serialize bus control modulation mix modulator ini")
+	)
+	static void SaveMixToProfile(const UObject* WorldContextObject, USoundControlBusMix* BusMix, int32 ProfileIndex = 0);
+
+	/** Loads control bus mix from a profile into UObject mix definition, deserialized from an ini file.
+	 * @param BusMix - Mix object to deserialize profile .ini to.
+	 * @param bActivate - If true, activate mix upon loading from profile.
+	 * @param ProfileIndex - Index of profile, allowing multiple profiles to be loaded to single mix object. If <= 0, loads from default profile (no suffix).
+	 * @return Channels - Channel values loaded from profile (empty if profile did not exist or had no values serialized).
+	 */
+	UFUNCTION(BlueprintCallable, Category = "Audio", DisplayName = "Load Control Bus Mix From Profile", meta = (
+		WorldContext = "WorldContextObject",
+		AdvancedDisplay = "2",
+		Keywords = "load deserialize control bus modulation mix modulator ini")
+	)
+	static UPARAM(DisplayName = "Channels") TArray<FSoundControlBusMixChannel> LoadMixFromProfile(const UObject* WorldContextObject, USoundControlBusMix* BusMix, bool bActivate = true, int32 ProfileIndex = 0);
+
+	/** Sets a mix with the provided channel data if channels provided in active instance proxy of mix. Does not update UObject definition of mix.
+	 * @param Mix - Mix to update
+	 * @param Channels - Channels to set.  If channel's bus is not referenced by mix, channel's update request is ignored.
+	 * @param bUpdateObject - If true, will dirty mix object and update channels on the SoundControlBusMix object in addition
+	 * to updating the audio thread proxy.
+	 */
+	UFUNCTION(BlueprintCallable, Category = "Audio", DisplayName = "Set Control Bus Mix", meta = (
+		WorldContext = "WorldContextObject",
+		Keywords = "set bus control modulation modulator mix channel")
+	)
+	static void UpdateMix(const UObject* WorldContextObject, USoundControlBusMix* Mix, TArray<FSoundControlBusMixChannel> Channels);
+
+	/** Sets filtered channels of a given class to a provided target value for active instance of mix. Does not update UObject definition of mix.
+	 * @param Mix - Mix to modify
+	 * @param AddressFilter - Address filter to apply to provided mix's channels.
+	 * @param BusClass - Filters buses by subclass.
+	 * @param Value - Target value to mix filtered channels to.
+	 * @param AttackTime - If non-negative, updates the attack time for the resulting bus channels found matching the provided filter.
+	 * @param ReleaseTime - If non-negative, updates the release time for the resulting bus channels found matching the provided filter.
+	 * @param bUpdateObject - If true, will dirty mix object and update channels on the SoundControlBusMix object in addition to updating 
+	 * the audio thread proxy.
+	 */
+	UFUNCTION(BlueprintCallable, Category = "Audio", DisplayName = "Set Control Bus Mix By Filter", meta = (
+		AdvancedDisplay = "5",
+		WorldContext = "WorldContextObject",
+		Keywords = "set bus control class modulation modulator mix channel value filter")
+	)
+	static void UpdateMixByFilter(
+		const UObject*						WorldContextObject,
+		USoundControlBusMix*				Mix,
+		FString								AddressFilter,
+		TSubclassOf<USoundControlBusBase>	BusClassFilter,
+		float								Value,
+		float								AttackTime	= -1.0f,
+		float								ReleaseTime = -1.0f);
+
+	/** Commits updates from a UObject definition of a modulator (e.g. Bus, Bus Mix, LFO) to active instance in audio thread
+	 * (ignored if modulator type has not been activated).
+	 * @param Modulator - Modulator to update
+	 */
+	UFUNCTION(BlueprintCallable, Category = "Audio", DisplayName = "Update Modulator", meta = (
+		WorldContext = "WorldContextObject",
+		Keywords = "update set control bus mix modulation modulator")
+	)
+	static void UpdateModulator(const UObject* WorldContextObject, USoundModulatorBase* Modulator);
+};
diff --git a/Engine/Plugins/Runtime/AudioModulation/Source/AudioModulation/Public/AudioModulationStyle.h b/Engine/Plugins/Runtime/AudioModulation/Source/AudioModulation/Public/AudioModulationStyle.h
new file mode 100644
--- /dev/null
+++ b/Engine/Plugins/Runtime/AudioModulation/Source/AudioModulation/Public/AudioModulationStyle.h
@@ -0,0 +1,30 @@
+// Copyright Epic Games, Inc. All Rights Reserved.
+#pragma once
+
+#include "CoreMinimal.h"
+#include "Kismet/BlueprintFunctionLibrary.h"
+#include "Math/Color.h"
+
+#include "AudioModulationStyle.generated.h"
+
+UCLASS()
+class AUDIOMODULATION_API UAudioModulationStyle : public UBlueprintFunctionLibrary
+{
+public:
+	GENERATED_BODY()
+
+	UFUNCTION(BlueprintCallable, Category = "Audio|Modulation|Style")
+	static const FColor GetVolumeBusColor() { return FColor(33, 183, 0); }
+
+	UFUNCTION(BlueprintCallable, Category = "Audio|Modulation|Style")
+	static const FColor GetPitchBusColor() { return FColor(181, 21, 0); }
+
+	UFUNCTION(BlueprintCallable, Category = "Audio|Modulation|Style")
+	static const FColor GetLPFBusColor() { return FColor(0, 156, 183); }
+
+	UFUNCTION(BlueprintCallable, Category = "Audio|Modulation|Style")
+	static const FColor GetHPFBusColor() { return FColor(94, 237, 183); }
+
+	UFUNCTION(BlueprintCallable, Category = "Audio|Modulation|Style")
+	static const FColor GetControlBusColor() { return FColor(215, 180, 210); }
+};
\ No newline at end of file
diff --git a/Engine/Plugins/Runtime/AudioModulation/Source/AudioModulationEditor/Private/AssetTypeActions/AssetTypeActions_SoundControlBus.h b/Engine/Plugins/Runtime/AudioModulation/Source/AudioModulationEditor/Private/AssetTypeActions/AssetTypeActions_SoundControlBus.h
--- a/Engine/Plugins/Runtime/AudioModulation/Source/AudioModulationEditor/Private/AssetTypeActions/AssetTypeActions_SoundControlBus.h
+++ b/Engine/Plugins/Runtime/AudioModulation/Source/AudioModulationEditor/Private/AssetTypeActions/AssetTypeActions_SoundControlBus.h
@@ -11,7 +11,7 @@ class FAssetTypeActions_SoundVolumeControlBus : public FAssetTypeActions_Base
 public:
 	// IAssetTypeActions Implementation
 	virtual FText GetName() const override { return NSLOCTEXT("AssetTypeActions", "AssetTypeActions_SoundVolumeControlBus", "Control Bus (Volume)"); }
-	virtual FColor GetTypeColor() const override { return FAudioModulationStyle::GetVolumeBusColor(); }
+	virtual FColor GetTypeColor() const override { return UAudioModulationStyle::GetVolumeBusColor(); }
 	virtual UClass* GetSupportedClass() const override;
 	virtual uint32 GetCategories() override { return EAssetTypeCategories::Sounds; }
 
@@ -23,7 +23,7 @@ class FAssetTypeActions_SoundPitchControlBus : public FAssetTypeActions_Base
 public:
 	// IAssetTypeActions Implementation
 	virtual FText GetName() const override { return NSLOCTEXT("AssetTypeActions", "AssetTypeActions_SoundPitchControlBus", "Control Bus (Pitch)"); }
-	virtual FColor GetTypeColor() const override { return FAudioModulationStyle::GetPitchBusColor(); }
+	virtual FColor GetTypeColor() const override { return UAudioModulationStyle::GetPitchBusColor(); }
 	virtual UClass* GetSupportedClass() const override;
 	virtual uint32 GetCategories() override { return EAssetTypeCategories::Sounds; }
 
@@ -35,7 +35,7 @@ class FAssetTypeActions_SoundLPFControlBus : public FAssetTypeActions_Base
 public:
 	// IAssetTypeActions Implementation
 	virtual FText GetName() const override { return NSLOCTEXT("AssetTypeActions", "AssetTypeActions_SoundLPFControlBus", "Control Bus (LPF)"); }
-	virtual FColor GetTypeColor() const override { return FAudioModulationStyle::GetLPFBusColor(); }
+	virtual FColor GetTypeColor() const override { return UAudioModulationStyle::GetLPFBusColor(); }
 	virtual UClass* GetSupportedClass() const override;
 	virtual uint32 GetCategories() override { return EAssetTypeCategories::Sounds; }
 
@@ -47,7 +47,7 @@ class FAssetTypeActions_SoundHPFControlBus : public FAssetTypeActions_Base
 public:
 	// IAssetTypeActions Implementation
 	virtual FText GetName() const override { return NSLOCTEXT("AssetTypeActions", "AssetTypeActions_SoundHPFControlBus", "Control Bus (HPF)"); }
-	virtual FColor GetTypeColor() const override { return FAudioModulationStyle::GetHPFBusColor(); }
+	virtual FColor GetTypeColor() const override { return UAudioModulationStyle::GetHPFBusColor(); }
 	virtual UClass* GetSupportedClass() const override;
 	virtual uint32 GetCategories() override { return EAssetTypeCategories::Sounds; }
 
@@ -59,7 +59,7 @@ class FAssetTypeActions_SoundControlBus : public FAssetTypeActions_Base
 public:
 	// IAssetTypeActions Implementation
 	virtual FText GetName() const override { return NSLOCTEXT("AssetTypeActions", "AssetTypeActions_SoundControlBus", "Control Bus"); }
-	virtual FColor GetTypeColor() const override { return FAudioModulationStyle::GetControlBusColor(); }
+	virtual FColor GetTypeColor() const override { return UAudioModulationStyle::GetControlBusColor(); }
 	virtual UClass* GetSupportedClass() const override;
 	virtual uint32 GetCategories() override { return EAssetTypeCategories::Sounds; }
 
diff --git a/Engine/Plugins/Runtime/AudioModulation/Source/AudioModulationEditor/Private/AudioModulationStyle.h b/Engine/Plugins/Runtime/AudioModulation/Source/AudioModulationEditor/Private/AudioModulationStyle.h
deleted file mode 100644
diff --git a/Engine/Plugins/Runtime/AudioModulation/Source/AudioModulationEditor/Private/Editors/ModulationSettingsCurveEditorViewStacked.cpp b/Engine/Plugins/Runtime/AudioModulation/Source/AudioModulationEditor/Private/Editors/ModulationSettingsCurveEditorViewStacked.cpp
--- a/Engine/Plugins/Runtime/AudioModulation/Source/AudioModulationEditor/Private/Editors/ModulationSettingsCurveEditorViewStacked.cpp
+++ b/Engine/Plugins/Runtime/AudioModulation/Source/AudioModulationEditor/Private/Editors/ModulationSettingsCurveEditorViewStacked.cpp
@@ -40,7 +40,7 @@ void FModCurveEditorModel::Refresh(EModSettingsEditorCurveOutput InCurveOutput,
 	{
 		case EModSettingsEditorCurveOutput::Volume:
 		{
-			Color = FAudioModulationStyle::GetVolumeBusColor();
+			Color = UAudioModulationStyle::GetVolumeBusColor();
 			ShortNameBase = LOCTEXT("ModulationOutputVolumeBus", "Volume");
 			IntentionName = TEXT("AudioVolume");
 		}
@@ -48,7 +48,7 @@ void FModCurveEditorModel::Refresh(EModSettingsEditorCurveOutput InCurveOutput,
 
 		case EModSettingsEditorCurveOutput::Pitch:
 		{
-			Color = FAudioModulationStyle::GetPitchBusColor();
+			Color = UAudioModulationStyle::GetPitchBusColor();
 			ShortNameBase = LOCTEXT("ModulationOutputPitchBus", "Pitch");
 			IntentionName = TEXT("AudioPitch");
 		}
@@ -56,7 +56,7 @@ void FModCurveEditorModel::Refresh(EModSettingsEditorCurveOutput InCurveOutput,
 
 		case EModSettingsEditorCurveOutput::Highpass:
 		{
-			Color = FAudioModulationStyle::GetHPFBusColor();
+			Color = UAudioModulationStyle::GetHPFBusColor();
 			ShortNameBase = LOCTEXT("ModulationOutputHPFBus", "Highpass");
 			IntentionName = TEXT("AudioHPFFreq");
 		}
@@ -64,7 +64,7 @@ void FModCurveEditorModel::Refresh(EModSettingsEditorCurveOutput InCurveOutput,
 
 		case EModSettingsEditorCurveOutput::Lowpass:
 		{
-			Color = FAudioModulationStyle::GetLPFBusColor();
+			Color = UAudioModulationStyle::GetLPFBusColor();
 			ShortNameBase = LOCTEXT("ModulationOutputLPFBus", "Lowpass");
 			IntentionName = TEXT("AudioLPFFreq");
 		}
@@ -72,7 +72,7 @@ void FModCurveEditorModel::Refresh(EModSettingsEditorCurveOutput InCurveOutput,
 
 		case EModSettingsEditorCurveOutput::Control:
 		{
-			Color = FAudioModulationStyle::GetControlBusColor();
+			Color = UAudioModulationStyle::GetControlBusColor();
 			ShortNameBase = FText::FromName(*InControlName);
 			IntentionName = TEXT("AudioControlValue");
 		}
diff --git a/Engine/Plugins/Runtime/AudioModulation/Source/AudioModulationEditor/Private/Layouts/SoundControlModulationPatchLayout.cpp b/Engine/Plugins/Runtime/AudioModulation/Source/AudioModulationEditor/Private/Layouts/SoundControlModulationPatchLayout.cpp
--- a/Engine/Plugins/Runtime/AudioModulation/Source/AudioModulationEditor/Private/Layouts/SoundControlModulationPatchLayout.cpp
+++ b/Engine/Plugins/Runtime/AudioModulation/Source/AudioModulationEditor/Private/Layouts/SoundControlModulationPatchLayout.cpp
@@ -7,8 +7,6 @@
 #include "DetailWidgetRow.h"
 #include "EditorStyleSet.h"
 #include "IAudioExtensionPlugin.h"
-#include "IDetailChildrenBuilder.h"
-#include "IDetailPropertyRow.h"
 #include "SoundModulationPatch.h"
 #include "Widgets/Input/SButton.h"
 #include "Widgets/Text/STextBlock.h"
diff --git a/Engine/Plugins/Runtime/AudioModulation/Source/AudioModulationEditor/Private/Layouts/SoundControlModulationPatchLayout.h b/Engine/Plugins/Runtime/AudioModulation/Source/AudioModulationEditor/Private/Layouts/SoundControlModulationPatchLayout.h
--- a/Engine/Plugins/Runtime/AudioModulation/Source/AudioModulationEditor/Private/Layouts/SoundControlModulationPatchLayout.h
+++ b/Engine/Plugins/Runtime/AudioModulation/Source/AudioModulationEditor/Private/Layouts/SoundControlModulationPatchLayout.h
@@ -4,6 +4,7 @@
 #include "CoreMinimal.h"
 
 #include "IDetailChildrenBuilder.h"
+#include "IDetailPropertyRow.h"
 #include "IPropertyTypeCustomization.h"
 #include "Layout/Visibility.h"
 #include "Misc/Attribute.h"
