Shadow variable fix

diff --git a/Engine/Plugins/Runtime/AudioModulation/Source/AudioModulation/Private/AudioModulationImpl.cpp b/Engine/Plugins/Runtime/AudioModulation/Source/AudioModulation/Private/AudioModulationImpl.cpp
--- a/Engine/Plugins/Runtime/AudioModulation/Source/AudioModulation/Private/AudioModulationImpl.cpp
+++ b/Engine/Plugins/Runtime/AudioModulation/Source/AudioModulation/Private/AudioModulationImpl.cpp
@@ -521,7 +521,6 @@ namespace AudioModulation
 
 			const FBusMixId MixId = static_cast<FBusMixId>(MixToSerialize->GetUniqueID());
 			const FString   MixName = MixToSerialize->GetName();
-			const FString   MixPath = MixToSerialize->GetPathName();
 
 			FBusMixHandle MixHandle = FBusMixHandle::Get(MixId, RefProxies.BusMixes);
 			const bool bIsActive = MixHandle.IsValid();
@@ -544,7 +543,7 @@ namespace AudioModulation
 				PassedChannelInfo.Add(Pair.Key, Channel.Value);
 			}
 
-			AsyncTask(ENamedThreads::GameThread, [this, PassedChannelInfo, MixId, MixToSerialize, InProfileIndex]()
+			AsyncTask(ENamedThreads::GameThread, [this, PassedChannelInfo, MixToSerialize, InProfileIndex]()
 			{
 				if (!MixToSerialize.IsValid())
 				{
@@ -571,7 +570,7 @@ namespace AudioModulation
 					}
 				}
 
-				const FString MixPath = MixToSerialize->GetPathName();
+				const FString   MixPath = MixToSerialize->GetPathName();
 				AudioModulation::FProfileSerializer::Serialize(*TempMix, InProfileIndex, &MixPath);
 			});
 		});
