Edigrate CLs 11902874, 11849778
Fix Mod mixes not reactivating correctly due to original channel settings not being passed when calling 'SetEnabled' on mix.  Fixed by moving 'Init' logic to SetEnabled and making public for AudioModulationImpl to call when reactivating a stopping mix.

- Pass modulation implementation to proxies to ensure version in charge of creation is that used for initialization/modification where necessary
- Fix leak and make ref counting safer on modulation proxies by moving inc/dec ref out of factory creaters/getters

#rnx

diff --git a/Engine/Plugins/Runtime/AudioModulation/Source/AudioModulation/Private/AudioModulationImpl.cpp b/Engine/Plugins/Runtime/AudioModulation/Source/AudioModulation/Private/AudioModulationImpl.cpp
--- a/Engine/Plugins/Runtime/AudioModulation/Source/AudioModulation/Private/AudioModulationImpl.cpp
+++ b/Engine/Plugins/Runtime/AudioModulation/Source/AudioModulation/Private/AudioModulationImpl.cpp
@@ -94,7 +94,7 @@ namespace AudioModulation
 			{
 				if (SourceSetting.GetId() == SettingsId)
 				{
-					SourceSetting = FModulationSettingsProxy(*Settings, RefProxies);
+					SourceSetting = FModulationSettingsProxy(*Settings, RefProxies, *this);
 				}
 			}
 
@@ -102,7 +102,7 @@ namespace AudioModulation
 			{
 				if (Pair.Value.GetId() == SettingsId)
 				{
-					Pair.Value = FModulationSettingsProxy(*Settings, RefProxies);
+					Pair.Value = FModulationSettingsProxy(*Settings, RefProxies, *this);
 				}
 			}
 		});
@@ -117,7 +117,7 @@ namespace AudioModulation
 		if (!SoundSettings.Contains(SoundId))
 		{
 			const USoundModulationSettings* Settings = CastChecked<USoundModulationSettings>(&InSettings);
-			SoundSettings.Add(SoundId, FModulationSettingsProxy(*Settings, RefProxies));
+			SoundSettings.Add(SoundId, FModulationSettingsProxy(*Settings, RefProxies, *this));
 		}
 	}
 
@@ -127,7 +127,7 @@ namespace AudioModulation
 
 		if (const USoundModulationSettings* Settings = CastChecked<USoundModulationSettings>(&InSettings))
 		{
-			SourceSettings[InSourceId] = FModulationSettingsProxy(*Settings, RefProxies);
+			SourceSettings[InSourceId] = FModulationSettingsProxy(*Settings, RefProxies, *this);
 		}
 	}
 
@@ -178,7 +178,7 @@ namespace AudioModulation
 					NewProxy.InitLFOs(*BusPtr.Get(), RefProxies.LFOs);
 				};
 
-				FBusHandle BusHandle = FBusHandle::Create(*BusPtr.Get(), RefProxies.Buses, OnCreate);
+				FBusHandle BusHandle = FBusHandle::Create(*BusPtr.Get(), RefProxies.Buses, *this, OnCreate);
 				ManuallyActivatedBuses.Add(MoveTemp(BusHandle));
 			}
 		});
@@ -195,11 +195,11 @@ namespace AudioModulation
 				FBusMixHandle BusMixHandle = FBusMixHandle::Get(*BusMixPtr.Get(), RefProxies.BusMixes);
 				if (BusMixHandle.IsValid())
 				{
-					BusMixHandle.FindProxy().SetEnabled();
+					BusMixHandle.FindProxy().SetEnabled(*BusMixPtr.Get());
 				}
 				else
 				{
-					BusMixHandle = FBusMixHandle::Create(*BusMixPtr.Get(), RefProxies.BusMixes);
+					BusMixHandle = FBusMixHandle::Create(*BusMixPtr.Get(), RefProxies.BusMixes, *this);
 				}
 
 				ManuallyActivatedBusMixes.Add(MoveTemp(BusMixHandle));
@@ -215,7 +215,7 @@ namespace AudioModulation
 		{
 			if (LFOPtr.IsValid())
 			{
-				FLFOHandle LFOHandle = FLFOHandle::Create(*LFOPtr.Get(), RefProxies.LFOs);
+				FLFOHandle LFOHandle = FLFOHandle::Create(*LFOPtr.Get(), RefProxies.LFOs, *this);
 				ManuallyActivatedLFOs.Add(MoveTemp(LFOHandle));
 			}
 		});
@@ -317,7 +317,7 @@ namespace AudioModulation
 		check(IsInAudioThread());
 
 		const USoundModulationSettings* Settings = CastChecked<USoundModulationSettings>(&SettingsBase);
-		FModulationPatchProxy VolumePatch(Settings->Volume, RefProxies);
+		FModulationPatchProxy VolumePatch(Settings->Volume, RefProxies, *this);
 
 		return CalculateModulationValue(VolumePatch);
 	}
@@ -330,7 +330,7 @@ namespace AudioModulation
 			if (BusPtr.IsValid())
 			{
 				const USoundControlBusBase& Bus = *BusPtr.Get();
-				FBusHandle BusHandle = FBusHandle::Create(Bus, RefProxies.Buses);
+				FBusHandle BusHandle = FBusHandle::Create(Bus, RefProxies.Buses, *this);
 				ManuallyActivatedBuses.Remove(MoveTemp(BusHandle));
 			}
 		});
@@ -362,7 +362,7 @@ namespace AudioModulation
 			if (LFOPtr.IsValid())
 			{
 				const USoundBusModulatorLFO& LFO = *LFOPtr.Get();
-				FLFOHandle LFOHandle = FLFOHandle::Create(LFO, RefProxies.LFOs);
+				FLFOHandle LFOHandle = FLFOHandle::Create(LFO, RefProxies.LFOs, *this);
 				ManuallyActivatedLFOs.Remove(MoveTemp(LFOHandle));
 			}
 		});
diff --git a/Engine/Plugins/Runtime/AudioModulation/Source/AudioModulation/Private/SoundControlBus.cpp b/Engine/Plugins/Runtime/AudioModulation/Source/AudioModulation/Private/SoundControlBus.cpp
--- a/Engine/Plugins/Runtime/AudioModulation/Source/AudioModulation/Private/SoundControlBus.cpp
+++ b/Engine/Plugins/Runtime/AudioModulation/Source/AudioModulation/Private/SoundControlBus.cpp
@@ -130,8 +130,8 @@ namespace AudioModulation
 	{
 	}
 
-	FControlBusProxy::FControlBusProxy(const USoundControlBusBase& InBus)
-		: TModulatorProxyRefType(InBus.GetName(), InBus.GetUniqueID())
+	FControlBusProxy::FControlBusProxy(const USoundControlBusBase& InBus, FAudioModulationImpl& InModulationImpl)
+		: TModulatorProxyRefType(InBus.GetName(), InBus.GetUniqueID(), InModulationImpl)
 	{
 		Init(InBus);
 	}
@@ -198,7 +198,8 @@ namespace AudioModulation
 		{
 			if (const USoundBusModulatorLFO* LFO = Cast<USoundBusModulatorLFO>(Modulator))
 			{
-				LFOHandles.Add(FLFOHandle::Create(*LFO, OutActiveLFOs));
+				check(ModulationImpl);
+				LFOHandles.Add(FLFOHandle::Create(*LFO, OutActiveLFOs, *ModulationImpl));
 			}
 		}
 	}
diff --git a/Engine/Plugins/Runtime/AudioModulation/Source/AudioModulation/Private/SoundControlBusMix.cpp b/Engine/Plugins/Runtime/AudioModulation/Source/AudioModulation/Private/SoundControlBusMix.cpp
--- a/Engine/Plugins/Runtime/AudioModulation/Source/AudioModulation/Private/SoundControlBusMix.cpp
+++ b/Engine/Plugins/Runtime/AudioModulation/Source/AudioModulation/Private/SoundControlBusMix.cpp
@@ -94,27 +94,24 @@ void USoundControlBusMix::SaveMixToProfile()
 
 namespace AudioModulation
 {
-	FModulatorBusMixChannelProxy::FModulatorBusMixChannelProxy(const FSoundControlBusMixChannel& Channel)
+	FModulatorBusMixChannelProxy::FModulatorBusMixChannelProxy(const FSoundControlBusMixChannel& Channel, FAudioModulationImpl& InModulationImpl)
 		: TModulatorProxyBase<FBusId>(Channel.Bus->GetName(), Channel.Bus->GetUniqueID())
 		, Address(Channel.Bus->Address)
 		, ClassId(Channel.Bus->GetClass()->GetUniqueID())
 		, Value(Channel.Value)
+		, BusHandle(FBusHandle::Create(*Channel.Bus, InModulationImpl.RefProxies.Buses, InModulationImpl))
 	{
-		check(Channel.Bus);
-		AudioModulation::FAudioModulationImpl* ModulationImpl = UAudioModulationStatics::GetModulationImpl(Channel.Bus->GetWorld());
-		check(ModulationImpl);
-		BusHandle = FBusHandle::Create(*Channel.Bus, ModulationImpl->RefProxies.Buses);
 	}
 
-	FModulatorBusMixProxy::FModulatorBusMixProxy(const USoundControlBusMix& InBusMix)
-		: TModulatorProxyRefType(InBusMix.GetName(), InBusMix.GetUniqueID())
+	FModulatorBusMixProxy::FModulatorBusMixProxy(const USoundControlBusMix& InBusMix, FAudioModulationImpl& InModulationImpl)
+		: TModulatorProxyRefType(InBusMix.GetName(), InBusMix.GetUniqueID(), InModulationImpl)
 	{
-		Init(InBusMix);
+		SetEnabled(InBusMix);
 	}
 
 	FModulatorBusMixProxy& FModulatorBusMixProxy::operator =(const USoundControlBusMix& InBusMix)
 	{
-		Init(InBusMix);
+		SetEnabled(InBusMix);
 
 		return *this;
 	}
@@ -129,7 +126,7 @@ namespace AudioModulation
 		Channels.Reset();
 	}
 
-	void FModulatorBusMixProxy::Init(const USoundControlBusMix& InBusMix)
+	void FModulatorBusMixProxy::SetEnabled(const USoundControlBusMix& InBusMix)
 	{
 		FChannelMap CachedChannels = Channels;
 		Channels.Reset();
@@ -139,14 +136,16 @@ namespace AudioModulation
 		{
 			if (Channel.Bus)
 			{
-				auto BusId = static_cast<const AudioModulation::FBusId>(Channel.Bus->GetUniqueID());
-				
-				FModulatorBusMixChannelProxy ChannelProxy(Channel);
+				auto BusId = static_cast<const FBusId>(Channel.Bus->GetUniqueID());
+
+				check(ModulationImpl);
+				FModulatorBusMixChannelProxy ChannelProxy(Channel, *ModulationImpl);
+
 				if (const FModulatorBusMixChannelProxy* CachedChannel = CachedChannels.Find(BusId))
 				{
 					ChannelProxy.Value.SetCurrentValue(CachedChannel->Value.GetCurrentValue());
 				}
-				
+
 				Channels.Emplace(BusId, ChannelProxy);
 			}
 			else
@@ -207,11 +206,6 @@ namespace AudioModulation
 		}
 	}
 
-	void FModulatorBusMixProxy::SetEnabled()
-	{
-		Status = EStatus::Enabled;
-	}
-
 	void FModulatorBusMixProxy::SetStopping()
 	{
 		if (Status == EStatus::Enabled)
diff --git a/Engine/Plugins/Runtime/AudioModulation/Source/AudioModulation/Private/SoundModulationPatch.cpp b/Engine/Plugins/Runtime/AudioModulation/Source/AudioModulation/Private/SoundModulationPatch.cpp
--- a/Engine/Plugins/Runtime/AudioModulation/Source/AudioModulation/Private/SoundModulationPatch.cpp
+++ b/Engine/Plugins/Runtime/AudioModulation/Source/AudioModulation/Private/SoundModulationPatch.cpp
@@ -241,7 +241,7 @@ namespace AudioModulation
 	{
 	}
 
-	FModulationInputProxy::FModulationInputProxy(const FSoundModulationInputBase& Input, FReferencedProxies& OutRefProxies)
+	FModulationInputProxy::FModulationInputProxy(const FSoundModulationInputBase& Input, FReferencedProxies& OutRefProxies, FAudioModulationImpl& ModulationImpl)
 		: Transform(Input.Transform)
 		, bSampleAndHold(Input.bSampleAndHold)
 	{
@@ -252,7 +252,7 @@ namespace AudioModulation
 			{
 				NewProxy.InitLFOs(*Bus, RefProxies->LFOs);
 			};
-			BusHandle = FBusHandle::Create(*Bus, OutRefProxies.Buses, OnCreate);
+			BusHandle = FBusHandle::Create(*Bus, OutRefProxies.Buses, ModulationImpl, OnCreate);
 		}
 	}
 
@@ -277,7 +277,7 @@ namespace AudioModulation
 	{
 	}
 
-	FModulationPatchProxy::FModulationPatchProxy(const FSoundModulationPatchBase& Patch, FReferencedProxies& OutRefProxies)
+	FModulationPatchProxy::FModulationPatchProxy(const FSoundModulationPatchBase& Patch, FReferencedProxies& OutRefProxies, FAudioModulationImpl& InModulationImpl)
 		: DefaultInputValue(Patch.DefaultInputValue)
 		, bBypass(Patch.bBypass)
 		, OutputProxy(*Patch.GetOutput())
@@ -285,7 +285,7 @@ namespace AudioModulation
 		TArray<const FSoundModulationInputBase*> Inputs = Patch.GetInputs();
 		for (const FSoundModulationInputBase* Input : Inputs)
 		{
-			InputProxies.Emplace_GetRef(*Input, OutRefProxies);
+			InputProxies.Emplace_GetRef(*Input, OutRefProxies, InModulationImpl);
 		}
 	}
 
@@ -306,16 +306,16 @@ namespace AudioModulation
 		Highpass.OutputProxy.Operator = ESoundModulatorOperator::Max;
 	}
 
-	FModulationSettingsProxy::FModulationSettingsProxy(const USoundModulationSettings& Settings, FReferencedProxies& OutRefProxies)
+	FModulationSettingsProxy::FModulationSettingsProxy(const USoundModulationSettings& Settings, FReferencedProxies& OutRefProxies, FAudioModulationImpl& InModulationImpl)
 		: TModulatorProxyBase<uint32>(Settings.GetName(), Settings.GetUniqueID())
-		, Volume(Settings.Volume, OutRefProxies)
-		, Pitch(Settings.Pitch, OutRefProxies)
-		, Lowpass(Settings.Lowpass, OutRefProxies)
-		, Highpass(Settings.Highpass, OutRefProxies)
+		, Volume(Settings.Volume, OutRefProxies, InModulationImpl)
+		, Pitch(Settings.Pitch, OutRefProxies, InModulationImpl)
+		, Lowpass(Settings.Lowpass, OutRefProxies, InModulationImpl)
+		, Highpass(Settings.Highpass, OutRefProxies, InModulationImpl)
 	{
 		for (const FSoundControlModulationPatch& Patch : Settings.Controls)
 		{
-			Controls.Add(Patch.Control, FModulationPatchProxy(Patch, OutRefProxies));
+			Controls.Add(Patch.Control, FModulationPatchProxy(Patch, OutRefProxies, InModulationImpl));
 		}
 	}
 } // namespace AudioModulation
\ No newline at end of file
diff --git a/Engine/Plugins/Runtime/AudioModulation/Source/AudioModulation/Private/SoundModulationProxy.h b/Engine/Plugins/Runtime/AudioModulation/Source/AudioModulation/Private/SoundModulationProxy.h
--- a/Engine/Plugins/Runtime/AudioModulation/Source/AudioModulation/Private/SoundModulationProxy.h
+++ b/Engine/Plugins/Runtime/AudioModulation/Source/AudioModulation/Private/SoundModulationProxy.h
@@ -18,6 +18,7 @@
 namespace AudioModulation
 {
 	// Forward Declarations
+	class FAudioModulationImpl;
 	struct FReferencedProxies;
 
 	// Modulator Ids
@@ -81,13 +82,17 @@ namespace AudioModulation
 			: Id(InId)
 			, ProxyMap(&InProxyMap)
 		{
+			if (IsValid())
+			{
+				ProxyType& Proxy = ProxyMap->FindChecked(Id);
+				Proxy.IncRef();
+			}
 		}
 
 		TProxyHandle(const TProxyHandle& InHandle)
+			: Id(InHandle.Id)
+			, ProxyMap(InHandle.ProxyMap)
 		{
-			Id = InHandle.Id;
-			ProxyMap = InHandle.ProxyMap;
-
 			if (IsValid())
 			{
 				ProxyType& Proxy = ProxyMap->FindChecked(Id);
@@ -96,27 +101,31 @@ namespace AudioModulation
 		}
 
 		TProxyHandle(TProxyHandle&& InHandle)
+			: Id(InHandle.Id)
+			, ProxyMap(InHandle.ProxyMap)
 		{
-			Id = InHandle.Id;
-			ProxyMap = InHandle.ProxyMap;
-
+			// No Inc/DecRef required due to transfer of ownership
 			InHandle.Id = static_cast<IdType>(INDEX_NONE);
 			InHandle.ProxyMap = nullptr;
 		}
 
 		~TProxyHandle()
 		{
-			if (ProxyMap)
+			if (!ProxyMap || !IsValid())
+			{
+				return;
+			}
+
+			ProxyType& Proxy = ProxyMap->FindChecked(Id);
+			const uint32 RefCount = Proxy.DecRef();
+			if (RefCount == 0)
 			{
-				if (ProxyType* Proxy = ProxyMap->Find(Id))
-				{
-					const uint32 RefCount = Proxy->DecRef();
-					if (RefCount == 0)
-					{
-						UE_LOG(LogAudioModulation, Verbose, TEXT("Proxy released: Id '%u'."), Id);
-						ProxyMap->Remove(Id);
-					}
-				}
+#if UE_BUILD_SHIPPING
+				UE_LOG(LogAudioModulation, Verbose, TEXT("Proxy released: Id '%u'."), Id);
+#else // UE_BUILD_SHIPPING
+				UE_LOG(LogAudioModulation, Verbose, TEXT("Proxy '%s' released: Id '%u'."), *Proxy.GetName(), Id);
+#endif // !UE_BUILD_SHIPPING
+				ProxyMap->Remove(Id);
 			}
 		}
 
@@ -127,9 +136,8 @@ namespace AudioModulation
 		{
 			if (ProxyType* Proxy = InProxyMap.Find(ObjectId))
 			{
+				check(Proxy->ModulationImpl);
 				TProxyHandle<IdType, ProxyType, ProxyUObjType> NewHandle(ObjectId, InProxyMap);
-				Proxy->IncRef();
-
 				return MoveTemp(NewHandle);
 			} 
 
@@ -142,49 +150,38 @@ namespace AudioModulation
 		static TProxyHandle<IdType, ProxyType, ProxyUObjType> Get(const ProxyUObjType& InObject, TMap<IdType, ProxyType>& InProxyMap)
 		{
 			const IdType ObjectId = static_cast<IdType>(InObject.GetUniqueID());
-			if (ProxyType* Proxy = InProxyMap.Find(ObjectId))
-			{
-				TProxyHandle<IdType, ProxyType, ProxyUObjType> NewHandle(ObjectId, InProxyMap);
-				Proxy->IncRef();
-
-				return MoveTemp(NewHandle);
-			} 
-
-			return TProxyHandle<IdType, ProxyType, ProxyUObjType>();
+			return Get(ObjectId, InProxyMap);
 		}
 
 		/*
 		 * Creates a handle to a proxy modulation object tracked in the provided InProxyMap.  Creates new proxy if it doesn't exist.
 		 */
-		static TProxyHandle<IdType, ProxyType, ProxyUObjType> Create(const ProxyUObjType& InObject, TMap<IdType, ProxyType>& InProxyMap)
+		static TProxyHandle<IdType, ProxyType, ProxyUObjType> Create(const ProxyUObjType& InObject, TMap<IdType, ProxyType>& InProxyMap, FAudioModulationImpl& InModulationImpl)
 		{
 			const IdType ObjectId = static_cast<IdType>(InObject.GetUniqueID());
-			TProxyHandle<IdType, ProxyType, ProxyUObjType> NewHandle(ObjectId, InProxyMap);
 
-			ProxyType* Proxy = InProxyMap.Find(NewHandle.Id);
-			if (!Proxy)
+			TProxyHandle<IdType, ProxyType, ProxyUObjType> NewHandle = Get(InObject, InProxyMap);
+			if (!NewHandle.IsValid())
 			{
 				UE_LOG(LogAudioModulation, Verbose, TEXT("Proxy created: Id '%u' for object '%s'."), NewHandle.Id, *InObject.GetName());
-				Proxy = &InProxyMap.Add(NewHandle.Id, ProxyType(InObject));
+				InProxyMap.Add(ObjectId, ProxyType(InObject, InModulationImpl));
+				NewHandle = TProxyHandle<IdType, ProxyType, ProxyUObjType>(ObjectId, InProxyMap);
 			}
-			Proxy->IncRef();
 
 			return NewHandle;
 		}
 
-		static TProxyHandle<IdType, ProxyType, ProxyUObjType> Create(const ProxyUObjType& InObject, TMap<IdType, ProxyType>& InProxyMap, TFunction<void(ProxyType&)> OnCreateProxy)
+		static TProxyHandle<IdType, ProxyType, ProxyUObjType> Create(const ProxyUObjType& InObject, TMap<IdType, ProxyType>& InProxyMap, FAudioModulationImpl& InModulationImpl, TFunction<void(ProxyType&)> OnCreateProxy)
 		{
 			const IdType ObjectId = static_cast<IdType>(InObject.GetUniqueID());
-			TProxyHandle<IdType, ProxyType, ProxyUObjType> NewHandle(ObjectId, InProxyMap);
-
-			ProxyType* Proxy = InProxyMap.Find(NewHandle.Id);
-			if (!Proxy)
+			TProxyHandle<IdType, ProxyType, ProxyUObjType> NewHandle = Get(InObject, InProxyMap);
+			if (!NewHandle.IsValid())
 			{
 				UE_LOG(LogAudioModulation, Verbose, TEXT("Proxy created: Id '%u' for object '%s'."), NewHandle.Id, *InObject.GetName());
-				Proxy = &InProxyMap.Add(NewHandle.Id, ProxyType(InObject));
-				OnCreateProxy(*Proxy);
+				InProxyMap.Add(ObjectId, ProxyType(InObject, InModulationImpl));
+				NewHandle = TProxyHandle<IdType, ProxyType, ProxyUObjType>(ObjectId, InProxyMap);
+				OnCreateProxy(NewHandle.FindProxy());
 			}
-			Proxy->IncRef();
 
 			return NewHandle;
 		}
@@ -196,41 +193,65 @@ namespace AudioModulation
 
 		FORCEINLINE TProxyHandle<IdType, ProxyType, ProxyUObjType>& operator =(const TProxyHandle<IdType, ProxyType, ProxyUObjType>& InHandle)
 		{
-			if (InHandle.Id == Id)
-			{
-				check(ProxyMap == InHandle.ProxyMap);
-				return *this;
-			}
-
+			// 1. If local proxy valid prior to move, cache to DecRef
+			ProxyType* ProxyToDecRef = nullptr;
 			if (IsValid())
 			{
-				ProxyType& Proxy = ProxyMap->FindChecked(Id);
-				Proxy.DecRef();
+				ProxyToDecRef = &ProxyMap->FindChecked(Id);
 			}
 
+			// 2. Move internal state
 			Id = InHandle.Id;
 			ProxyMap = InHandle.ProxyMap;
 
+			// 3. IncRef new data
 			if (IsValid())
 			{
 				ProxyType& Proxy = ProxyMap->FindChecked(Id);
 				Proxy.IncRef();
 			}
 
+			// 4. DecRef post IncRef to avoid premature destruction if handling same proxy
+			if (ProxyToDecRef)
+			{
+				ProxyToDecRef->DecRef();
+			}
+
 			return *this;
 		}
 
 		FORCEINLINE TProxyHandle<IdType, ProxyType, ProxyUObjType>& operator =(TProxyHandle<IdType, ProxyType, ProxyUObjType>&& InHandle)
 		{
+			// 1. If local proxy valid prior to move, cache to DecRef
+			ProxyType* ProxyToDecRef = nullptr;
 			if (IsValid())
 			{
-				ProxyType& Proxy = ProxyMap->FindChecked(Id);
-				Proxy.DecRef();
+				ProxyToDecRef = &ProxyMap->FindChecked(Id);
 			}
 
+			// 2. Move internal state
 			Id = InHandle.Id;
 			ProxyMap = InHandle.ProxyMap;
 
+			// 3. IncRef new data
+			if (IsValid())
+			{
+				ProxyType& Proxy = ProxyMap->FindChecked(Id);
+				Proxy.IncRef();
+			}
+
+			// 4. DecRef post IncRef to avoid premature destruction if handling same proxy
+			if (ProxyToDecRef)
+			{
+				ProxyToDecRef->DecRef();
+			}
+
+			// 5. DecRef handle to move and clear state
+			if (InHandle.IsValid())
+			{
+				ProxyType& Proxy = ProxyMap->FindChecked(InHandle.Id);
+				Proxy.DecRef();
+			}
 			InHandle.Id = static_cast<IdType>(INDEX_NONE);
 			InHandle.ProxyMap = nullptr;
 
@@ -292,43 +313,54 @@ namespace AudioModulation
 	{
 	protected:
 		uint32 RefCount;
+		FAudioModulationImpl* ModulationImpl;
 
 	public:
 		TModulatorProxyRefType()
 			: TModulatorProxyBase<IdType>()
 			, RefCount(0)
+			, ModulationImpl(nullptr)
 		{
 		}
 
-		TModulatorProxyRefType(const FString& InName, const uint32 InId)
+		TModulatorProxyRefType(const FString& InName, const uint32 InId, FAudioModulationImpl& InModulationImpl)
 			: TModulatorProxyBase<IdType>(InName, InId)
 			, RefCount(0)
+			, ModulationImpl(&InModulationImpl)
 		{
 		}
 
 		TModulatorProxyRefType(const TModulatorProxyRefType& InProxyRef)
 			: TModulatorProxyBase<IdType>(InProxyRef.GetName(), InProxyRef.GetId())
 			, RefCount(InProxyRef.RefCount)
+			, ModulationImpl(InProxyRef.ModulationImpl)
 		{
 		}
 
 		TModulatorProxyRefType(TModulatorProxyRefType&& InProxyRef)
 			: TModulatorProxyBase<IdType>(InProxyRef.GetName(), InProxyRef.GetId())
 			, RefCount(InProxyRef.RefCount)
+			, ModulationImpl(InProxyRef.ModulationImpl)
 		{
 			InProxyRef.RefCount = 0;
+			InProxyRef.ModulationImpl = nullptr;
 		}
 
 		TModulatorProxyRefType& operator=(const TModulatorProxyRefType& InOther)
 		{
-			RefCount = 0;
+			RefCount = InOther.RefCount;
+			ModulationImpl = InOther.ModulationImpl;
 			return *this;
 		}
 
 		TModulatorProxyRefType& operator=(const TModulatorProxyRefType&& InOther)
 		{
 			RefCount = InOther.RefCount;
+			ModulationImpl = InOther.ModulationImpl;
+
 			InOther.RefCount = 0;
+			InOther.ModulationImpl = nullptr;
+
 			return *this;
 		}
 
@@ -364,7 +396,7 @@ namespace AudioModulation
 	{
 	public:
 		FModulatorLFOProxy();
-		FModulatorLFOProxy(const USoundBusModulatorLFO& InLFO);
+		FModulatorLFOProxy(const USoundBusModulatorLFO& InLFO, FAudioModulationImpl& InModulationImpl);
 
 		FModulatorLFOProxy& operator =(const USoundBusModulatorLFO& InLFO);
 
@@ -385,7 +417,7 @@ namespace AudioModulation
 	{
 	public:
 		FControlBusProxy();
-		FControlBusProxy(const USoundControlBusBase& Bus);
+		FControlBusProxy(const USoundControlBusBase& Bus, FAudioModulationImpl& InModulationImpl);
 
 		FControlBusProxy& operator =(const USoundControlBusBase& InLFO);
 
@@ -424,7 +456,7 @@ namespace AudioModulation
 
 	struct FModulatorBusMixChannelProxy : public TModulatorProxyBase<FBusId>
 	{
-		FModulatorBusMixChannelProxy(const FSoundControlBusMixChannel& Channel);
+		FModulatorBusMixChannelProxy(const FSoundControlBusMixChannel& Channel, FAudioModulationImpl& ModulationImpl);
 		FString Address;
 		uint32 ClassId;
 		FSoundModulationValue Value;
@@ -441,7 +473,7 @@ namespace AudioModulation
 			Stopped
 		};
 
-		FModulatorBusMixProxy(const USoundControlBusMix& Mix);
+		FModulatorBusMixProxy(const USoundControlBusMix& InMix, FAudioModulationImpl& InModulationImpl);
 
 		FModulatorBusMixProxy& operator =(const USoundControlBusMix& InBusMix);
 
@@ -450,7 +482,7 @@ namespace AudioModulation
 		// Resets channel map
 		void Reset();
 
-		void SetEnabled();
+		void SetEnabled(const USoundControlBusMix& InBusMix);
 		void SetMix(const TArray<FSoundControlBusMixChannel>& InChannels);
 		void SetMixByFilter(const FString& InAddressFilter, uint32 InFilterClassId, const FSoundModulationValue& InValue);
 		void SetStopping();
@@ -461,7 +493,6 @@ namespace AudioModulation
 		FChannelMap Channels;
 
 	private:
-		void Init(const USoundControlBusMix& InBusMix);
 		EStatus Status;
 	};
 
@@ -472,7 +503,7 @@ namespace AudioModulation
 	struct FModulationInputProxy
 	{
 		FModulationInputProxy();
-		FModulationInputProxy(const FSoundModulationInputBase& Patch, FReferencedProxies& OutRefProxies);
+		FModulationInputProxy(const FSoundModulationInputBase& Patch, FReferencedProxies& OutRefProxies, FAudioModulationImpl& ModulationImpl);
 
 		FBusHandle BusHandle;
 		FSoundModulationInputTransform Transform;
@@ -501,7 +532,7 @@ namespace AudioModulation
 	struct FModulationPatchProxy
 	{
 		FModulationPatchProxy();
-		FModulationPatchProxy(const FSoundModulationPatchBase& Patch, FReferencedProxies& OutRefProxies);
+		FModulationPatchProxy(const FSoundModulationPatchBase& Patch, FReferencedProxies& OutRefProxies, FAudioModulationImpl& InModulationImpl);
 
 		/** Default value of patch (Value mixed when inputs are provided or not, regardless of active state)*/
 		float DefaultInputValue;
@@ -519,7 +550,7 @@ namespace AudioModulation
 	struct FModulationSettingsProxy : public TModulatorProxyBase<uint32>
 	{
 		FModulationSettingsProxy();
-		FModulationSettingsProxy(const USoundModulationSettings& Settings, FReferencedProxies& OutRefProxies);
+		FModulationSettingsProxy(const USoundModulationSettings& Settings, FReferencedProxies& OutRefProxies, FAudioModulationImpl& InModulationImpl);
 
 		FModulationPatchProxy Volume;
 		FModulationPatchProxy Pitch;
diff --git a/Engine/Plugins/Runtime/AudioModulation/Source/AudioModulation/Private/SoundModulatorLFO.cpp b/Engine/Plugins/Runtime/AudioModulation/Source/AudioModulation/Private/SoundModulatorLFO.cpp
--- a/Engine/Plugins/Runtime/AudioModulation/Source/AudioModulation/Private/SoundModulatorLFO.cpp
+++ b/Engine/Plugins/Runtime/AudioModulation/Source/AudioModulation/Private/SoundModulatorLFO.cpp
@@ -54,8 +54,8 @@ namespace AudioModulation
 		LFO.Start();
 	}
 
-	FModulatorLFOProxy::FModulatorLFOProxy(const USoundBusModulatorLFO& InLFO)
-		: TModulatorProxyRefType(InLFO.GetName(), InLFO.GetUniqueID())
+	FModulatorLFOProxy::FModulatorLFOProxy(const USoundBusModulatorLFO& InLFO, FAudioModulationImpl& InModulationImpl)
+		: TModulatorProxyRefType(InLFO.GetName(), InLFO.GetUniqueID(), InModulationImpl)
 		, Offset(InLFO.Offset)
 		, Value(1.0f)
 		, bBypass(InLFO.bBypass)
