Fix issues with modulation output curves reordering randomly when modifying output curve values in editor & hide output curve when bypassed

diff --git a/Engine/Plugins/Runtime/AudioModulation/Source/AudioModulationEditor/Private/Editors/ModulationSettingsCurveEditorViewStacked.cpp b/Engine/Plugins/Runtime/AudioModulation/Source/AudioModulationEditor/Private/Editors/ModulationSettingsCurveEditorViewStacked.cpp
--- a/Engine/Plugins/Runtime/AudioModulation/Source/AudioModulationEditor/Private/Editors/ModulationSettingsCurveEditorViewStacked.cpp
+++ b/Engine/Plugins/Runtime/AudioModulation/Source/AudioModulationEditor/Private/Editors/ModulationSettingsCurveEditorViewStacked.cpp
@@ -9,6 +9,7 @@
 #include "EditorStyleSet.h"
 #include "Fonts/FontMeasure.h"
 #include "SCurveEditorPanel.h"
+#include "SoundModulationPatch.h"
 #include "Widgets/Text/STextBlock.h"
 
 
@@ -21,22 +22,25 @@ FModCurveEditorModel::FModCurveEditorModel(FRichCurve& InRichCurve, UObject* InO
 	, Source(InSource)
 {
 	SupportedViews = ViewId;
-	Refresh(EModSettingsEditorCurveOutput::Control, &InControlName, InSharedCurve);
+	Init(&InControlName, InSharedCurve);
 }
 
 FModCurveEditorModel::FModCurveEditorModel(FRichCurve& InRichCurve, UObject* InOwner, EModSettingsEditorCurveOutput InOutput, EModSettingsOutputEditorCurveSource InSource, UCurveFloat* InSharedCurve)
 	: FRichCurveEditorModelRaw(&InRichCurve, InOwner)
+	, Settings(Cast<USoundModulationSettings>(InOwner))
 	, Output(InOutput)
 	, Source(InSource)
 {
 	SupportedViews = ViewId;
-	Refresh(InOutput, nullptr /* InControlName */, InSharedCurve);
+	Init(nullptr /* InControlName */, InSharedCurve);
 }
 
-void FModCurveEditorModel::Refresh(EModSettingsEditorCurveOutput InCurveOutput, const FName* InControlName, UCurveFloat* InSharedCurve)
+void FModCurveEditorModel::Init(const FName* InControlName, UCurveFloat* InSharedCurve)
 {
+	bKeyDrawEnabled = true;
+
 	FText ShortNameBase;
-	switch (InCurveOutput)
+	switch (Output)
 	{
 		case EModSettingsEditorCurveOutput::Volume:
 		{
@@ -79,43 +83,95 @@ void FModCurveEditorModel::Refresh(EModSettingsEditorCurveOutput InCurveOutput,
 		break;
 	}
 
-	switch (Source)
+	const bool bIsBypassed = GetIsBypassed();
+	if (bIsBypassed)
+	{
+		SetShortDisplayName(FText::Format(LOCTEXT("ModulationCurveDisabledDisplayName", "{0} (Bypassed)"), ShortNameBase));
+		bKeyDrawEnabled = false;
+	}
+	else
 	{
-		case EModSettingsOutputEditorCurveSource::Shared:
+		switch (Source)
 		{
-			check(InSharedCurve);
-			FText CurveNameText = FText::FromString(InSharedCurve->GetName());
-			SetShortDisplayName(FText::Format(LOCTEXT("ModulationSharedDisplayName", "{0} (Shared - {1})"), ShortNameBase, CurveNameText));
+			case EModSettingsOutputEditorCurveSource::Shared:
+			{
+				check(InSharedCurve);
+				FText CurveNameText = FText::FromString(InSharedCurve->GetName());
+				SetShortDisplayName(FText::Format(LOCTEXT("ModulationSharedDisplayName", "{0} (Shared - {1})"), ShortNameBase, CurveNameText));
+			}
+			break;
+
+			case EModSettingsOutputEditorCurveSource::Custom:
+			{
+				ShortDisplayName = FText::Format(LOCTEXT("ModulationOutputCurveDisplayName", "{0} (Custom)"), ShortNameBase);
+			}
+			break;
+
+			case EModSettingsOutputEditorCurveSource::Expression:
+			{
+				bKeyDrawEnabled = false;
+				ShortDisplayName = FText::Format(LOCTEXT("ModulationOutputCurveExpressionDisplayName", "{0} (Expression)"), ShortNameBase);
+			}
+			break;
+
+			case EModSettingsOutputEditorCurveSource::Unset:
+			default:
+			{
+				bKeyDrawEnabled = false;
+				ShortDisplayName = FText::Format(LOCTEXT("ModulationOutputCurveUnsetDisplayName", "{0} (Shared - Unset)"), ShortNameBase);
+			}
+			break;
+		}
+	}
+}
+
+FLinearColor FModCurveEditorModel::GetColor() const
+{
+	return !GetIsBypassed() && (Source == EModSettingsOutputEditorCurveSource::Custom || Source == EModSettingsOutputEditorCurveSource::Expression)
+		? Color
+		: Color.Desaturate(0.45f);
+}
+
+bool FModCurveEditorModel::GetIsBypassed() const
+{
+	if (Settings.IsValid())
+	{
+		switch (Output)
+		{
+		case EModSettingsEditorCurveOutput::Volume:
+		{
+			return Settings->Volume.bBypass;
+		}
+		break;
+
+		case EModSettingsEditorCurveOutput::Pitch:
+		{
+			return Settings->Pitch.bBypass;
 		}
 		break;
 
-		case EModSettingsOutputEditorCurveSource::Custom:
+		case EModSettingsEditorCurveOutput::Highpass:
 		{
-			ShortDisplayName = FText::Format(LOCTEXT("ModulationOutputCurveDisplayName", "{0} (Custom)"), ShortNameBase);
+			return Settings->Highpass.bBypass;
 		}
 		break;
 
-		case EModSettingsOutputEditorCurveSource::Expression:
+		case EModSettingsEditorCurveOutput::Lowpass:
 		{
-			bKeyDrawEnabled = false;
-			ShortDisplayName = FText::Format(LOCTEXT("ModulationOutputCurveExpressionDisplayName", "{0} (Expression)"), ShortNameBase);
+			return Settings->Lowpass.bBypass;
 		}
 		break;
 
-		case EModSettingsOutputEditorCurveSource::Unset:
+		case EModSettingsEditorCurveOutput::Control:
 		default:
 		{
-			bKeyDrawEnabled = false;
-			ShortDisplayName = FText::Format(LOCTEXT("ModulationOutputCurveUnsetDisplayName", "{0} (Shared - Unset)"), ShortNameBase);
+			return false;
+		}
+		break;
 		}
 	}
-}
 
-FLinearColor FModCurveEditorModel::GetColor() const
-{
-	return Source == EModSettingsOutputEditorCurveSource::Custom || Source == EModSettingsOutputEditorCurveSource::Expression
-		? Color
-		: Color.Desaturate(0.45f);
+	return false;
 }
 
 EModSettingsEditorCurveOutput FModCurveEditorModel::GetOutput() const
diff --git a/Engine/Plugins/Runtime/AudioModulation/Source/AudioModulationEditor/Private/Editors/ModulationSettingsCurveEditorViewStacked.h b/Engine/Plugins/Runtime/AudioModulation/Source/AudioModulationEditor/Private/Editors/ModulationSettingsCurveEditorViewStacked.h
--- a/Engine/Plugins/Runtime/AudioModulation/Source/AudioModulationEditor/Private/Editors/ModulationSettingsCurveEditorViewStacked.h
+++ b/Engine/Plugins/Runtime/AudioModulation/Source/AudioModulationEditor/Private/Editors/ModulationSettingsCurveEditorViewStacked.h
@@ -5,11 +5,15 @@
 #include "Curves/RichCurve.h"
 #include "RichCurveEditorModel.h"
 #include "UObject/Object.h"
+#include "UObject/WeakObjectPtr.h"
+#include "UObject/WeakObjectPtrTemplates.h"
 #include "Views/SCurveEditorViewStacked.h"
 #include "Views/SInteractiveCurveEditorView.h"
 
+
 // Forward Declarations
 class UCurveFloat;
+class USoundModulationSettings;
 
 
 enum class EModSettingsEditorCurveOutput : uint8
@@ -146,15 +150,18 @@ public:
 	FModCurveEditorModel(FRichCurve& InRichCurve, UObject* InOwner, FName InControlName, EModSettingsOutputEditorCurveSource InSource, UCurveFloat* SharedCurve);
 	FModCurveEditorModel(FRichCurve& InRichCurve, UObject* InOwner, EModSettingsEditorCurveOutput InOutput, EModSettingsOutputEditorCurveSource InSource, UCurveFloat* SharedCurve);
 
+	bool GetIsBypassed() const;
+
 	virtual bool IsReadOnly() const override;
 
 	virtual FLinearColor GetColor() const override;
 	EModSettingsEditorCurveOutput GetOutput() const;
 	EModSettingsOutputEditorCurveSource GetSource() const;
 
-	void Refresh(EModSettingsEditorCurveOutput InCurveOutput, const FName* InControlName, UCurveFloat* InSharedCurve);
-
 private:
+	void Init(const FName* InControlName, UCurveFloat* InSharedCurve);
+
+	TWeakObjectPtr<USoundModulationSettings> Settings;
 
 	EModSettingsEditorCurveOutput Output;
 	EModSettingsOutputEditorCurveSource Source;
diff --git a/Engine/Plugins/Runtime/AudioModulation/Source/AudioModulationEditor/Private/Editors/ModulationSettingsEditor.cpp b/Engine/Plugins/Runtime/AudioModulation/Source/AudioModulationEditor/Private/Editors/ModulationSettingsEditor.cpp
--- a/Engine/Plugins/Runtime/AudioModulation/Source/AudioModulationEditor/Private/Editors/ModulationSettingsEditor.cpp
+++ b/Engine/Plugins/Runtime/AudioModulation/Source/AudioModulationEditor/Private/Editors/ModulationSettingsEditor.cpp
@@ -169,45 +169,48 @@ void FModulationSettingsEditor::GenerateExpressionCurve(const FSoundModulationOu
 	TSharedPtr<FRichCurve> NewCurve = MakeShared<FRichCurve>();
 	ExpressionCurves[InCurveIndex] = NewCurve;
 
-	int32 CurveResolution;
-	switch (InOutput.Transform.Curve)
+	if (!GetIsBypassed(InCurveIndex))
 	{
-		case ESoundModulatorOutputCurve::Linear:
+		int32 CurveResolution;
+		switch (InOutput.Transform.Curve)
 		{
-			CurveResolution = 2;
-		}
-		break;
+			case ESoundModulatorOutputCurve::Linear:
+			{
+				CurveResolution = 2;
+			}
+			break;
 
-		case ESoundModulatorOutputCurve::Sin:
-		case ESoundModulatorOutputCurve::SCurve:
-		{
-			CurveResolution = 64;
-		}
-		break;
+			case ESoundModulatorOutputCurve::Sin:
+			case ESoundModulatorOutputCurve::SCurve:
+			{
+				CurveResolution = 64;
+			}
+			break;
 
-		case ESoundModulatorOutputCurve::Log:
-		case ESoundModulatorOutputCurve::Exp:
-		case ESoundModulatorOutputCurve::Exp_Inverse:
-		{
-			CurveResolution = 256;
-		}
-		break;
+			case ESoundModulatorOutputCurve::Log:
+			case ESoundModulatorOutputCurve::Exp:
+			case ESoundModulatorOutputCurve::Exp_Inverse:
+			{
+				CurveResolution = 256;
+			}
+			break;
 
-		default:
-		{
-			CurveResolution = 128;
+			default:
+			{
+				CurveResolution = 128;
+			}
+			break;
 		}
-		break;
-	}
 
-	const float CurveResolutionRatio = 1.0f / (CurveResolution - 1);
-	for (int32 i = 0; i < CurveResolution; ++i)
-	{
-		const float X = FMath::Lerp(InOutput.Transform.InputMin, InOutput.Transform.InputMax, CurveResolutionRatio * i);
-		float Y = X;
-		InOutput.Transform.Apply(Y);
+		const float CurveResolutionRatio = 1.0f / (CurveResolution - 1);
+		for (int32 i = 0; i < CurveResolution; ++i)
+		{
+			const float X = FMath::Lerp(InOutput.Transform.InputMin, InOutput.Transform.InputMax, CurveResolutionRatio * i);
+			float Y = X;
+			InOutput.Transform.Apply(Y);
 
-		NewCurve->AddKey(X, Y);
+			NewCurve->AddKey(X, Y);
+		}
 	}
 
 	const EModSettingsOutputEditorCurveSource Source = bIsUnset ? EModSettingsOutputEditorCurveSource::Unset : EModSettingsOutputEditorCurveSource::Expression;
@@ -218,54 +221,29 @@ void FModulationSettingsEditor::SetCurveAtOrderIndex(int32 InCurveIndex, FRichCu
 {
 	check(CurveEditor.IsValid());
 
-	bool bIsNew = false;
-	while (InCurveIndex >= CurveModels.Num())
+	if (CurveModels.Num() <= InCurveIndex)
 	{
-		CurveModels.Add(FCurveModelID::Unique());
-		bIsNew = true;
+		CurveModels.SetNumZeroed(InCurveIndex + 1);
 	}
+	CurveModels[InCurveIndex] = FCurveModelID::Unique();
 
-	// Find if the incoming rich curve is already set.  If so, early out
-	if (!bIsNew)
+	// Finally, create a FCurveModel add the new FRichCurve
+	if (!GetIsBypassed(InCurveIndex))
 	{
-		FCurveModel* CurrentModel = CurveEditor->FindCurve(CurveModels[InCurveIndex]);
-		if (CurrentModel && CurrentModel->GetCurve() == &InRichCurve)
+		if (InCurveIndex < static_cast<int32>(EModSettingsEditorCurveOutput::Control))
 		{
-			FModCurveEditorModel* ModModel = static_cast<FModCurveEditorModel*>(CurrentModel);
-			if (InCurveIndex < static_cast<int32>(EModSettingsEditorCurveOutput::Control))
-			{
-				ModModel->Refresh(static_cast<EModSettingsEditorCurveOutput>(InCurveIndex), nullptr, InSharedCurve);
-			}
-			else
-			{
-				check(!InSharedCurve);
-
-				FName ControlName = GetControlName(InCurveIndex);
-				ModModel->Refresh(EModSettingsEditorCurveOutput::Control, &ControlName, InSharedCurve);
-			}
-			CurveEditor->PinCurve(CurveModels[InCurveIndex]);
-			return;
+			const EModSettingsEditorCurveOutput CurveOutput = static_cast<EModSettingsEditorCurveOutput>(InCurveIndex);
+			TUniquePtr<FModCurveEditorModel> NewCurve = MakeUnique<FModCurveEditorModel>(InRichCurve, GetEditingObject(), CurveOutput, InSource, InSharedCurve);
+			CurveModels[InCurveIndex] = CurveEditor->AddCurve(MoveTemp(NewCurve));
+		}
+		else
+		{
+			const USoundModulationSettings* ModSettings = CastChecked<USoundModulationSettings>(GetEditingObject());
+			const int32 ControlIndex = GetControlIndex(InCurveIndex);
+			const FName Name = ModSettings->Controls[ControlIndex].Control;
+			TUniquePtr<FModCurveEditorModel> NewCurve = MakeUnique<FModCurveEditorModel>(InRichCurve, GetEditingObject(), Name, InSource, InSharedCurve);
+			CurveModels[InCurveIndex] = CurveEditor->AddCurve(MoveTemp(NewCurve));
 		}
-	}
-
-	// Curves don't match, so remove what's already there
-	if (FCurveModel* CurveModel = CurveEditor->FindCurve(CurveModels[InCurveIndex]))
-	{
-		CurveEditor->RemoveCurve(CurveModels[InCurveIndex]);
-	}
-
-	// Finally, create a FCurveModel add the new FRichCurve
-	if (InCurveIndex < static_cast<int32>(EModSettingsEditorCurveOutput::Control))
-	{
-		const EModSettingsEditorCurveOutput CurveOutput = static_cast<EModSettingsEditorCurveOutput>(InCurveIndex);
-		TUniquePtr<FModCurveEditorModel> NewCurve = MakeUnique<FModCurveEditorModel>(InRichCurve, GetEditingObject(), CurveOutput, InSource, InSharedCurve);
-		CurveModels[InCurveIndex] = CurveEditor->AddCurve(MoveTemp(NewCurve));
-	}
-	else
-	{
-		const FName Name = CastChecked<USoundModulationSettings>(GetEditingObject())->Controls[GetControlIndex(InCurveIndex)].Control;
-		TUniquePtr<FModCurveEditorModel> NewCurve = MakeUnique<FModCurveEditorModel>(InRichCurve, GetEditingObject(), Name, InSource, InSharedCurve);
-		CurveModels[InCurveIndex] = CurveEditor->AddCurve(MoveTemp(NewCurve));
 	}
 
 	CurveEditor->PinCurve(CurveModels[InCurveIndex]);
@@ -347,6 +325,50 @@ FName FModulationSettingsEditor::GetControlName(int32 InCurveIndex) const
 	return Settings->Controls[ControlIndex].Control;
 }
 
+bool FModulationSettingsEditor::GetIsBypassed(int32 InCurveIndex) const
+{
+	if (USoundModulationSettings* Settings = Cast<USoundModulationSettings>(GetEditingObject()))
+	{
+		const EModSettingsEditorCurveOutput CurveOutput = static_cast<EModSettingsEditorCurveOutput>(InCurveIndex);
+
+		switch (CurveOutput)
+		{
+		case EModSettingsEditorCurveOutput::Volume:
+		{
+			return Settings->Volume.bBypass;
+		}
+		break;
+
+		case EModSettingsEditorCurveOutput::Pitch:
+		{
+			return Settings->Pitch.bBypass;
+		}
+		break;
+
+		case EModSettingsEditorCurveOutput::Highpass:
+		{
+			return Settings->Highpass.bBypass;
+		}
+		break;
+
+		case EModSettingsEditorCurveOutput::Lowpass:
+		{
+			return Settings->Lowpass.bBypass;
+		}
+		break;
+
+		case EModSettingsEditorCurveOutput::Control:
+		default:
+		{
+			return false;
+		}
+		break;
+		}
+	}
+
+	return false;
+}
+
 FSoundModulationOutputBase* FModulationSettingsEditor::FindModulationOutput(int32 InCurveIndex)
 {
 	USoundModulationSettings* Settings = Cast<USoundModulationSettings>(GetEditingObject());
@@ -485,6 +507,13 @@ void FModulationSettingsEditor::UpdateCurve(int32 InCurveIndex)
 		return;
 	}
 
+	// Unpin all curves already created while iterating curves being updated, to ensure
+	// they're repinned in the same order
+	if (InCurveIndex < CurveModels.Num())
+	{
+		CurveEditor->UnpinCurve(CurveModels[InCurveIndex]);
+	}
+
 	switch (Output->Transform.Curve)
 	{
 		case ESoundModulatorOutputCurve::Exp:
@@ -535,13 +564,17 @@ void FModulationSettingsEditor::UpdateCurves()
 		return;
 	}
 
-	for (int32 i = 0; i < GetCurveCount(); ++i)
+	CurveEditor->RemoveAllCurves();
+
+	int32 CurveCount = GetCurveCount();
+	for (int32 i = 0; i < CurveCount; ++i)
 	{
 		UpdateCurve(i);
 	}
 
 	// Remove old curves from editor
-	for (int32 i = GetCurveCount(); i < CurveModels.Num(); ++i)
+	CurveCount = GetCurveCount();
+	for (int32 i = CurveCount; i < CurveModels.Num(); ++i)
 	{
 		CurveEditor->RemoveCurve(CurveModels[i]);
 	}
diff --git a/Engine/Plugins/Runtime/AudioModulation/Source/AudioModulationEditor/Private/Editors/ModulationSettingsEditor.h b/Engine/Plugins/Runtime/AudioModulation/Source/AudioModulationEditor/Private/Editors/ModulationSettingsEditor.h
--- a/Engine/Plugins/Runtime/AudioModulation/Source/AudioModulationEditor/Private/Editors/ModulationSettingsEditor.h
+++ b/Engine/Plugins/Runtime/AudioModulation/Source/AudioModulationEditor/Private/Editors/ModulationSettingsEditor.h
@@ -68,6 +68,8 @@ private:
 
 	int32 GetCurveCount() const;
 
+	bool GetIsBypassed(int32 InCurveIndex) const;
+
 	/** Get the orientation for the snap value controls. */
 	EOrientation GetSnapLabelOrientation() const;
 
