Fix bad move/add merge
#rnx

diff --git a/Engine/Plugins/Runtime/AudioModulation/Source/AudioModulation/Public/AudioModulationStyle.h b/Engine/Plugins/Runtime/AudioModulation/Source/AudioModulation/Public/AudioModulationStyle.h
--- a/Engine/Plugins/Runtime/AudioModulation/Source/AudioModulation/Public/AudioModulationStyle.h
+++ b/Engine/Plugins/Runtime/AudioModulation/Source/AudioModulation/Public/AudioModulationStyle.h
@@ -2,15 +2,29 @@
 #pragma once
 
 #include "CoreMinimal.h"
+#include "Kismet/BlueprintFunctionLibrary.h"
 #include "Math/Color.h"
 
+#include "AudioModulationStyle.generated.h"
 
-class FAudioModulationStyle
+UCLASS()
+class AUDIOMODULATION_API UAudioModulationStyle : public UBlueprintFunctionLibrary
 {
 public:
-	static const FColor GetVolumeBusColor()  { return FColor(33, 183, 0); }
-	static const FColor GetPitchBusColor()   { return FColor(181, 21, 0); }
-	static const FColor GetLPFBusColor()     { return FColor(0, 156, 183); }
-	static const FColor GetHPFBusColor()     { return FColor(94, 237, 183); }
+	GENERATED_BODY()
+
+	UFUNCTION(BlueprintCallable, Category = "Audio|Modulation|Style")
+	static const FColor GetVolumeBusColor() { return FColor(33, 183, 0); }
+
+	UFUNCTION(BlueprintCallable, Category = "Audio|Modulation|Style")
+	static const FColor GetPitchBusColor() { return FColor(181, 21, 0); }
+
+	UFUNCTION(BlueprintCallable, Category = "Audio|Modulation|Style")
+	static const FColor GetLPFBusColor() { return FColor(0, 156, 183); }
+
+	UFUNCTION(BlueprintCallable, Category = "Audio|Modulation|Style")
+	static const FColor GetHPFBusColor() { return FColor(94, 237, 183); }
+
+	UFUNCTION(BlueprintCallable, Category = "Audio|Modulation|Style")
 	static const FColor GetControlBusColor() { return FColor(215, 180, 210); }
 };
\ No newline at end of file
