Move icons to AudioModulation plugin content directory
#rnx

diff --git a/Engine/Plugins/Runtime/AudioModulation/AudioModulation.uplugin b/Engine/Plugins/Runtime/AudioModulation/AudioModulation.uplugin
--- a/Engine/Plugins/Runtime/AudioModulation/AudioModulation.uplugin
+++ b/Engine/Plugins/Runtime/AudioModulation/AudioModulation.uplugin
@@ -11,7 +11,7 @@
 	"MarketplaceURL": "",
 	"SupportURL": "",
 	"EnabledByDefault": false,
-	"CanContainContent": false,
+	"CanContainContent": true,
 	"IsBetaVersion": true,
 	"Installed": false,
 	"Modules":
diff --git a/Engine/Plugins/Runtime/AudioModulation/Source/AudioModulationEditor/Private/AudioModulationEditor.cpp b/Engine/Plugins/Runtime/AudioModulation/Source/AudioModulationEditor/Private/AudioModulationEditor.cpp
--- a/Engine/Plugins/Runtime/AudioModulation/Source/AudioModulationEditor/Private/AudioModulationEditor.cpp
+++ b/Engine/Plugins/Runtime/AudioModulation/Source/AudioModulationEditor/Private/AudioModulationEditor.cpp
@@ -46,7 +46,7 @@ void FAudioModulationEditorModule::SetIcon(const FString& ClassName)
 	static const FVector2D Icon16 = FVector2D(16.0f, 16.0f);
 	static const FVector2D Icon64 = FVector2D(64.0f, 64.0f);
 
-	static const FString IconDir = FPaths::EngineDir() / FString(TEXT("Plugins/Runtime/AudioModulation/Icons"));
+	static const FString IconDir = FPaths::EnginePluginsDir() / FString(TEXT("Runtime/AudioModulation/Content"));
 
 	const FString IconFileName16 = FString::Printf(TEXT("%s_16x.png"), *ClassName);
 	const FString IconFileName64 = FString::Printf(TEXT("%s_64x.png"), *ClassName);
