Generic Parameter Modulation Refactor
- Add classes to allow modulation of effect parameters
- Remove cruft from first attempt at adding generic modulators

Minor-ish:
- Clean-up headers
- Rename AudioModulationImpl --> AudioModulationSystem
- Remove dead generic code

diff --git a/Engine/Plugins/Runtime/AudioModulation/Source/AudioModulation/Private/AudioModulation.cpp b/Engine/Plugins/Runtime/AudioModulation/Source/AudioModulation/Private/AudioModulation.cpp
--- a/Engine/Plugins/Runtime/AudioModulation/Source/AudioModulation/Private/AudioModulation.cpp
+++ b/Engine/Plugins/Runtime/AudioModulation/Source/AudioModulation/Private/AudioModulation.cpp
@@ -1,12 +1,13 @@
 // Copyright Epic Games, Inc. All Rights Reserved.
 #include "AudioModulation.h"
 
-#include "AudioModulationInternal.h"
+#include "AudioModulationSystem.h"
 #include "AudioModulationLogging.h"
 #include "CanvasTypes.h"
 #include "Features/IModularFeatures.h"
-#include "IAudioExtensionPlugin.h"
+#include "IAudioModulation.h"
 #include "Modules/ModuleManager.h"
+#include "SoundControlBusMix.h"
 #include "SoundModulationPatch.h"
 #include "UnrealClient.h"
 
@@ -18,78 +19,108 @@ namespace AudioModulation
 {
 	FAudioModulation::FAudioModulation()
 	{
-		Impl = TUniquePtr<FAudioModulationImpl>(new FAudioModulationImpl());
+		ModSystem = TUniquePtr<FAudioModulationSystem>(new FAudioModulationSystem());
 	}
 
-	FAudioModulationImpl* FAudioModulation::GetImpl()
+	FAudioModulationSystem* FAudioModulation::GetModulationSystem()
 	{
-		return Impl.Get();
+		return ModSystem.Get();
 	}
 
 	float FAudioModulation::CalculateInitialVolume(const USoundModulationPluginSourceSettingsBase& Settings)
 	{
-		return Impl->CalculateInitialVolume(Settings);
+		return ModSystem->CalculateInitialVolume(Settings);
 	}
 
 	void FAudioModulation::Initialize(const FAudioPluginInitializationParams& InitializationParams)
 	{
-		Impl->Initialize(InitializationParams);
+		ModSystem->Initialize(InitializationParams);
+	}
+
+	void FAudioModulation::OnBeginAudioRenderThreadUpdate()
+	{
+		ModSystem->OnBeginAudioRenderThreadUpdate();
 	}
 
 #if WITH_EDITOR
 	void FAudioModulation::OnEditPluginSettings(const USoundModulationPluginSourceSettingsBase& Settings)
 	{
-		Impl->OnEditPluginSettings(Settings);
+		ModSystem->OnEditPluginSettings(Settings);
 	}
 #endif // WITH_EDITOR
 
 	void FAudioModulation::OnInitSound(ISoundModulatable& Sound, const USoundModulationPluginSourceSettingsBase& Settings)
 	{
-		Impl->OnInitSound(Sound, Settings);
+		ModSystem->OnInitSound(Sound, Settings);
 	}
 
 	void FAudioModulation::OnInitSource(const uint32 SourceId, const FName& AudioComponentUserId, const uint32 NumChannels, const USoundModulationPluginSourceSettingsBase& Settings)
 	{
-		Impl->OnInitSource(SourceId, AudioComponentUserId, NumChannels, Settings);
+		ModSystem->OnInitSource(SourceId, AudioComponentUserId, NumChannels, Settings);
 	}
 
 	void FAudioModulation::OnReleaseSound(ISoundModulatable& Sound)
 	{
-		Impl->OnReleaseSound(Sound);
+		ModSystem->OnReleaseSound(Sound);
 	}
 
 	void FAudioModulation::OnReleaseSource(const uint32 SourceId)
 	{
-		Impl->OnReleaseSource(SourceId);
+		ModSystem->OnReleaseSource(SourceId);
 	}
 
 #if !UE_BUILD_SHIPPING
 	bool FAudioModulation::OnPostHelp(FCommonViewportClient* ViewportClient, const TCHAR* Stream)
 	{
-		return Impl->OnPostHelp(ViewportClient, Stream);
+		return ModSystem->OnPostHelp(ViewportClient, Stream);
 	}
 
 	int32 FAudioModulation::OnRenderStat(FViewport* Viewport, FCanvas* Canvas, int32 X, int32 Y, const UFont& Font, const FVector* ViewLocation, const FRotator* ViewRotation)
 	{
-		return Impl->OnRenderStat(Viewport, Canvas, X, Y, Font, ViewLocation, ViewRotation);
+		return ModSystem->OnRenderStat(Viewport, Canvas, X, Y, Font, ViewLocation, ViewRotation);
 	}
 
 	bool FAudioModulation::OnToggleStat(FCommonViewportClient* ViewportClient, const TCHAR* Stream)
 	{
-		return Impl->OnToggleStat(ViewportClient, Stream);
+		return ModSystem->OnToggleStat(ViewportClient, Stream);
 	}
 #endif // !UE_BUILD_SHIPPING
 
 	bool FAudioModulation::ProcessControls(const uint32 InSourceId, FSoundModulationControls& OutControls)
 	{
 		SCOPE_CYCLE_COUNTER(STAT_AudioModulationProcessControls);
-		return Impl->ProcessControls(InSourceId, OutControls);
+		return ModSystem->ProcessControls(InSourceId, OutControls);
 	}
 
 	void FAudioModulation::ProcessModulators(const float Elapsed)
 	{
 		SCOPE_CYCLE_COUNTER(STAT_AudioModulationProcessModulators);
-		Impl->ProcessModulators(Elapsed);
+		ModSystem->ProcessModulators(Elapsed);
+	}
+
+	bool FAudioModulation::RegisterModulator(uint32 InParentId, const USoundModulatorBase& InModulatorBase)
+	{
+		return ModSystem->RegisterModulator(InParentId, InModulatorBase);
+	}
+
+	bool FAudioModulation::RegisterModulator(uint32 InParentId, Audio::FModulatorId InModulatorId)
+	{
+		return ModSystem->RegisterModulator(InParentId, InModulatorId);
+	}
+
+	bool FAudioModulation::GetModulatorValue(const Audio::FModulatorHandle& ModulatorHandle, float& OutValue)
+	{
+		return ModSystem->GetModulatorValue(ModulatorHandle, OutValue);
+	}
+
+	void FAudioModulation::UnregisterModulator(const Audio::FModulatorHandle& InHandle)
+	{
+		ModSystem->UnregisterModulator(InHandle);
+	}
+
+	void FAudioModulation::UpdateModulator(const USoundModulatorBase& InModulator)
+	{
+		ModSystem->UpdateModulator(InModulator);
 	}
 } // namespace AudioModulation
 
diff --git a/Engine/Plugins/Runtime/AudioModulation/Source/AudioModulation/Private/AudioModulationDebugger.cpp b/Engine/Plugins/Runtime/AudioModulation/Source/AudioModulation/Private/AudioModulationDebugger.cpp
--- a/Engine/Plugins/Runtime/AudioModulation/Source/AudioModulation/Private/AudioModulationDebugger.cpp
+++ b/Engine/Plugins/Runtime/AudioModulation/Source/AudioModulation/Private/AudioModulationDebugger.cpp
@@ -5,11 +5,13 @@
 #if !UE_BUILD_SHIPPING
 #include "CoreMinimal.h"
 #include "AudioModulationLogging.h"
+#include "AudioModulationSystem.h"
 #include "AudioThread.h"
 #include "Engine/Canvas.h"
 #include "Engine/Engine.h"
-#include "IAudioExtensionPlugin.h"
+#include "IAudioModulation.h"
 #include "Misc/CoreDelegates.h"
+#include "SoundControlBusProxy.h"
 #include "SoundModulationPatch.h"
 #include "SoundModulationProxy.h"
 #include "SoundModulationValue.h"
@@ -314,7 +316,7 @@ namespace AudioModulation
 			FControlBusMixDebugInfo DebugInfo;
 			DebugInfo.Name = Proxy->GetName();
 			DebugInfo.RefCount = Proxy->GetRefCount();
-			for (const TPair< FBusId, FModulatorBusMixChannelProxy>& Channel : Proxy->Channels)
+			for (const TPair<FBusId, FModulatorBusMixChannelProxy>& Channel : Proxy->Channels)
 			{
 				FControlBusMixChannelDebugInfo ChannelDebugInfo;
 				ChannelDebugInfo.CurrentValue = Channel.Value.Value.GetCurrentValue();
diff --git a/Engine/Plugins/Runtime/AudioModulation/Source/AudioModulation/Private/AudioModulationImpl.cpp b/Engine/Plugins/Runtime/AudioModulation/Source/AudioModulation/Private/AudioModulationImpl.cpp
deleted file mode 100644
diff --git a/Engine/Plugins/Runtime/AudioModulation/Source/AudioModulation/Private/AudioModulationInternal.h b/Engine/Plugins/Runtime/AudioModulation/Source/AudioModulation/Private/AudioModulationInternal.h
deleted file mode 100644
diff --git a/Engine/Plugins/Runtime/AudioModulation/Source/AudioModulation/Private/AudioModulationProfileSerializer.h b/Engine/Plugins/Runtime/AudioModulation/Source/AudioModulation/Private/AudioModulationProfileSerializer.h
--- a/Engine/Plugins/Runtime/AudioModulation/Source/AudioModulation/Private/AudioModulationProfileSerializer.h
+++ b/Engine/Plugins/Runtime/AudioModulation/Source/AudioModulation/Private/AudioModulationProfileSerializer.h
@@ -7,6 +7,7 @@
 #include "Logging/LogMacros.h"
 #include "Misc/ConfigCacheIni.h"
 #include "SoundControlBusMix.h"
+#include "SoundModulationValue.h"
 #include "UObject/Object.h"
 #include "UObject/ObjectMacros.h"
 
diff --git a/Engine/Plugins/Runtime/AudioModulation/Source/AudioModulation/Private/AudioModulationStatics.cpp b/Engine/Plugins/Runtime/AudioModulation/Source/AudioModulation/Private/AudioModulationStatics.cpp
--- a/Engine/Plugins/Runtime/AudioModulation/Source/AudioModulation/Private/AudioModulationStatics.cpp
+++ b/Engine/Plugins/Runtime/AudioModulation/Source/AudioModulation/Private/AudioModulationStatics.cpp
@@ -4,9 +4,9 @@
 #include "Async/Async.h"
 #include "AudioDevice.h"
 #include "AudioModulation.h"
-#include "AudioModulationInternal.h"
 #include "AudioModulationLogging.h"
 #include "AudioModulationProfileSerializer.h"
+#include "AudioModulationSystem.h"
 #include "CoreGlobals.h"
 #include "Engine/Engine.h"
 #include "HAL/IConsoleManager.h"
@@ -92,7 +92,7 @@ static FAutoConsoleCommand GModulationLoadMixProfile(
 
 					if (bActivateUpdate)
 					{
-						UAudioModulationStatics::UpdateModulator(Mix, Mix);
+						UAudioModulationStatics::UpdateMixFromObject(Mix, Mix);
 					}
 					return;
 				}
@@ -121,9 +121,9 @@ namespace AudioModulation
 
 		if (Activate)
 		{
-			if (FAudioModulationImpl* ModulationImpl = UAudioModulationStatics::GetModulationImpl(World))
+			if (FAudioModulationSystem* ModSystem = UAudioModulationStatics::GetModulationSystem(World))
 			{
-				ModulationImpl->ActivateBus(*NewBus);
+				ModSystem->ActivateBus(*NewBus);
 			}
 		}
 
@@ -145,9 +145,9 @@ void UAudioModulationStatics::ActivateBus(const UObject* WorldContextObject, USo
 	}
 
 	UWorld* World = GetAudioWorld(WorldContextObject);
-	if (AudioModulation::FAudioModulationImpl* ModulationImpl = GetModulationImpl(World))
+	if (AudioModulation::FAudioModulationSystem* ModSystem = GetModulationSystem(World))
 	{
-		ModulationImpl->ActivateBus(*Bus);
+		ModSystem->ActivateBus(*Bus);
 	}
 }
 
@@ -156,9 +156,9 @@ void UAudioModulationStatics::ActivateBusMix(const UObject* WorldContextObject,
 	if (BusMix)
 	{
 		UWorld* World = GetAudioWorld(WorldContextObject);
-		if (AudioModulation::FAudioModulationImpl* ModulationImpl = GetModulationImpl(World))
+		if (AudioModulation::FAudioModulationSystem* ModSystem = GetModulationSystem(World))
 		{
-			ModulationImpl->ActivateBusMix(*BusMix);
+			ModSystem->ActivateBusMix(*BusMix);
 		}
 	}
 }
@@ -166,11 +166,11 @@ void UAudioModulationStatics::ActivateBusMix(const UObject* WorldContextObject,
 void UAudioModulationStatics::ActivateBusModulator(const UObject* WorldContextObject, USoundBusModulatorBase* Modulator)
 {
 	UWorld* World = GetAudioWorld(WorldContextObject);
-	if (AudioModulation::FAudioModulationImpl* ModulationImpl = GetModulationImpl(World))
+	if (AudioModulation::FAudioModulationSystem* ModSystem = GetModulationSystem(World))
 	{
 		if (USoundBusModulatorLFO* LFO = Cast<USoundBusModulatorLFO>(Modulator))
 		{
-			ModulationImpl->ActivateLFO(*LFO);
+			ModSystem->ActivateLFO(*LFO);
 		}
 	}
 }
@@ -191,7 +191,7 @@ UWorld* UAudioModulationStatics::GetAudioWorld(const UObject* WorldContextObject
 	return World;
 }
 
-AudioModulation::FAudioModulationImpl* UAudioModulationStatics::GetModulationImpl(UWorld* World)
+AudioModulation::FAudioModulationSystem* UAudioModulationStatics::GetModulationSystem(UWorld* World)
 {
 	FAudioDeviceHandle AudioDevice;
 	if (World)
@@ -210,7 +210,7 @@ AudioModulation::FAudioModulationImpl* UAudioModulationStatics::GetModulationImp
 	{
 		if (IAudioModulation* ModulationInterface = AudioDevice->ModulationInterface.Get())
 		{
-			return static_cast<AudioModulation::FAudioModulation*>(ModulationInterface)->GetImpl();
+			return static_cast<AudioModulation::FAudioModulation*>(ModulationInterface)->GetModulationSystem();
 		}
 	}
 
@@ -252,9 +252,9 @@ USoundBusModulatorLFO* UAudioModulationStatics::CreateLFO(const UObject* WorldCo
 
 	if (Activate)
 	{
-		if (AudioModulation::FAudioModulationImpl* ModulationImpl = GetModulationImpl(World))
+		if (AudioModulation::FAudioModulationSystem* ModSystem = GetModulationSystem(World))
 		{
-			ModulationImpl->ActivateLFO(*NewLFO);
+			ModSystem->ActivateLFO(*NewLFO);
 		}
 	}
 
@@ -294,9 +294,9 @@ USoundControlBusMix* UAudioModulationStatics::CreateBusMix(const UObject* WorldC
 
 	if (Activate)
 	{
-		if (AudioModulation::FAudioModulationImpl* ModulationImpl = GetModulationImpl(World))
+		if (AudioModulation::FAudioModulationSystem* ModSystem = GetModulationSystem(World))
 		{
-			ModulationImpl->ActivateBusMix(*NewBusMix);
+			ModSystem->ActivateBusMix(*NewBusMix);
 		}
 	}
 
@@ -308,9 +308,9 @@ void UAudioModulationStatics::DeactivateBus(const UObject* WorldContextObject, U
 	if (Bus)
 	{
 		UWorld* World = GetAudioWorld(WorldContextObject);
-		if (AudioModulation::FAudioModulationImpl* ModulationImpl = GetModulationImpl(World))
+		if (AudioModulation::FAudioModulationSystem* ModSystem = GetModulationSystem(World))
 		{
-			ModulationImpl->DeactivateBus(*Bus);
+			ModSystem->DeactivateBus(*Bus);
 		}
 	}
 }
@@ -320,9 +320,9 @@ void UAudioModulationStatics::DeactivateBusMix(const UObject* WorldContextObject
 	if (BusMix)
 	{
 		UWorld* World = GetAudioWorld(WorldContextObject);
-		if (AudioModulation::FAudioModulationImpl* ModulationImpl = GetModulationImpl(World))
+		if (AudioModulation::FAudioModulationSystem* ModSystem = GetModulationSystem(World))
 		{
-			ModulationImpl->DeactivateBusMix(*BusMix);
+			ModSystem->DeactivateBusMix(*BusMix);
 		}
 	}
 }
@@ -330,11 +330,11 @@ void UAudioModulationStatics::DeactivateBusMix(const UObject* WorldContextObject
 void UAudioModulationStatics::DeactivateBusModulator(const UObject* WorldContextObject, USoundBusModulatorBase* Modulator)
 {
 	UWorld* World = GetAudioWorld(WorldContextObject);
-	if (AudioModulation::FAudioModulationImpl* ModulationImpl = GetModulationImpl(World))
+	if (AudioModulation::FAudioModulationSystem* ModSystem = GetModulationSystem(World))
 	{
 		if (USoundBusModulatorLFO* LFO = Cast<USoundBusModulatorLFO>(Modulator))
 		{
-			ModulationImpl->DeactivateLFO(*LFO);
+			ModSystem->DeactivateLFO(*LFO);
 		}
 	}
 }
@@ -342,11 +342,11 @@ void UAudioModulationStatics::DeactivateBusModulator(const UObject* WorldContext
 void UAudioModulationStatics::SaveMixToProfile(const UObject* WorldContextObject, USoundControlBusMix* BusMix, int32 ProfileIndex)
 {
 	UWorld* World = GetAudioWorld(WorldContextObject);
-	if (AudioModulation::FAudioModulationImpl* ModulationImpl = GetModulationImpl(World))
+	if (AudioModulation::FAudioModulationSystem* ModSystem = GetModulationSystem(World))
 	{
 		if (BusMix)
 		{
-			return ModulationImpl->SaveMixToProfile(*BusMix, ProfileIndex);
+			return ModSystem->SaveMixToProfile(*BusMix, ProfileIndex);
 		}
 	}
 }
@@ -356,13 +356,13 @@ TArray<FSoundControlBusMixChannel> UAudioModulationStatics::LoadMixFromProfile(c
 	if (BusMix)
 	{
 		UWorld* World = GetAudioWorld(WorldContextObject);
-		if (AudioModulation::FAudioModulationImpl* ModulationImpl = GetModulationImpl(World))
+		if (AudioModulation::FAudioModulationSystem* ModSystem = GetModulationSystem(World))
 		{
 			if (bActivate)
 			{
 				ActivateBusMix(WorldContextObject, BusMix);
 			}
-			return ModulationImpl->LoadMixFromProfile(ProfileIndex, *BusMix);
+			return ModSystem->LoadMixFromProfile(ProfileIndex, *BusMix);
 		}
 	}
 
@@ -374,11 +374,11 @@ void UAudioModulationStatics::UpdateMix(const UObject* WorldContextObject, USoun
 	if (Mix)
 	{
 		UWorld* World = GetAudioWorld(WorldContextObject);
-		if (AudioModulation::FAudioModulationImpl* ModulationImpl = GetModulationImpl(World))
+		if (AudioModulation::FAudioModulationSystem* ModSystem = GetModulationSystem(World))
 		{
 			// UObject representation is not updated in this form of the call as doing so from
 			// PIE can result in an unstable state where UObject is modified but not properly dirtied.
-			ModulationImpl->UpdateMix(Channels, *Mix, false /* bUpdateObject */);
+			ModSystem->UpdateMix(Channels, *Mix, false /* bUpdateObject */);
 		}
 	}
 }
@@ -395,13 +395,25 @@ void UAudioModulationStatics::UpdateMixByFilter(
 	if (Mix)
 	{
 		UWorld* World = GetAudioWorld(WorldContextObject);
-		if (AudioModulation::FAudioModulationImpl* ModulationImpl = GetModulationImpl(World))
+		if (AudioModulation::FAudioModulationSystem* ModSystem = GetModulationSystem(World))
 		{
 			FSoundModulationValue ModValue(Value, AttackTime, ReleaseTime);
 
 			// UObject representation is not updated in this form of the call as doing so from
 			// PIE can result in an unstable state where UObject is modified but not properly dirtied.
-			ModulationImpl->UpdateMixByFilter(AddressFilter, BusClassFilter, ModValue, *Mix, false /* bUpdateObject */);
+			ModSystem->UpdateMixByFilter(AddressFilter, BusClassFilter, ModValue, *Mix, false /* bUpdateObject */);
+		}
+	}
+}
+
+void UAudioModulationStatics::UpdateMixFromObject(const UObject* WorldContextObject, USoundControlBusMix* Mix)
+{
+	if (Mix)
+	{
+		UWorld* World = GetAudioWorld(WorldContextObject);
+		if (AudioModulation::FAudioModulationSystem* ModSystem = GetModulationSystem(World))
+		{
+			ModSystem->UpdateMix(*Mix);
 		}
 	}
 }
@@ -411,9 +423,9 @@ void UAudioModulationStatics::UpdateModulator(const UObject* WorldContextObject,
 	if (Modulator)
 	{
 		UWorld* World = GetAudioWorld(WorldContextObject);
-		if (AudioModulation::FAudioModulationImpl* ModulationImpl = GetModulationImpl(World))
+		if (AudioModulation::FAudioModulationSystem* ModSystem = GetModulationSystem(World))
 		{
-			ModulationImpl->UpdateModulator(*Modulator);
+			ModSystem->UpdateModulator(*Modulator);
 		}
 	}
 }
diff --git a/Engine/Plugins/Runtime/AudioModulation/Source/AudioModulation/Private/AudioModulationSystem.cpp b/Engine/Plugins/Runtime/AudioModulation/Source/AudioModulation/Private/AudioModulationSystem.cpp
new file mode 100644
--- /dev/null
+++ b/Engine/Plugins/Runtime/AudioModulation/Source/AudioModulation/Private/AudioModulationSystem.cpp
@@ -0,0 +1,919 @@
+// Copyright Epic Games, Inc. All Rights Reserved.
+#include "AudioModulationSystem.h"
+
+#if WITH_AUDIOMODULATION
+#include "Audio/AudioAddressPattern.h"
+#include "AudioModulationLogging.h"
+#include "AudioModulationProfileSerializer.h"
+#include "AudioThread.h"
+#include "Engine/Engine.h"
+#include "IAudioModulation.h"
+#include "Misc/CoreDelegates.h"
+#include "SoundControlBusProxy.h"
+#include "SoundControlBusMixProxy.h"
+#include "SoundModulatorLFOProxy.h"
+#include "SoundModulationPatchProxy.h"
+#include "SoundModulationProxy.h"
+#include "UObject/UObjectIterator.h"
+#include "UObject/WeakObjectPtr.h"
+
+
+#if !UE_BUILD_SHIPPING
+#include "AudioModulationDebugger.h"
+#include "Async/Async.h"
+#endif // !UE_BUILD_SHIPPING
+
+DECLARE_DWORD_COUNTER_STAT(TEXT("Bus Count"),	STAT_AudioModulationBusCount, STATGROUP_AudioModulation)
+DECLARE_DWORD_COUNTER_STAT(TEXT("LFO Count"),	STAT_AudioModulationLFOCount, STATGROUP_AudioModulation)
+DECLARE_DWORD_COUNTER_STAT(TEXT("Mix Count"),	STAT_AudioModulationMixCount, STATGROUP_AudioModulation)
+DECLARE_DWORD_COUNTER_STAT(TEXT("Patch Count"), STAT_AudioModulationPatchCount, STATGROUP_AudioModulation)
+
+
+namespace AudioModulation
+{
+	struct FProfileChannelInfo
+	{
+		USoundControlBusBase* Bus;
+		FSoundModulationValue Value;
+
+		FProfileChannelInfo(const FModulatorBusMixChannelProxy& InProxy)
+			: Bus(nullptr)
+			, Value(InProxy.Value)
+		{
+		}
+	};
+
+	FAudioModulationSystem::FAudioModulationSystem()
+	{
+	}
+
+	void FAudioModulationSystem::Initialize(const FAudioPluginInitializationParams& InitializationParams)
+	{
+		SourceSettings.AddDefaulted(InitializationParams.NumSources);
+	}
+
+#if WITH_EDITOR
+	void FAudioModulationSystem::OnEditPluginSettings(const USoundModulationPluginSourceSettingsBase& InSettings)
+	{
+		const TWeakObjectPtr<const USoundModulationPluginSourceSettingsBase> SettingsPtr = &InSettings;
+		RunCommandOnAudioThread([this, SettingsPtr]()
+		{
+			if (!SettingsPtr.IsValid())
+			{
+				return;
+			}
+
+			const USoundModulationSettings* Settings = CastChecked<USoundModulationSettings>(SettingsPtr.Get());
+			const uint32 SettingsId = Settings->GetUniqueID();
+			for (FModulationSettingsProxy& SourceSetting : SourceSettings)
+			{
+				if (SourceSetting.GetId() == SettingsId)
+				{
+					SourceSetting = FModulationSettingsProxy(*Settings, *this);
+				}
+			}
+
+			for (TPair<uint32, FModulationSettingsProxy>& Pair : SoundSettings)
+			{
+				if (Pair.Value.GetId() == SettingsId)
+				{
+					Pair.Value = FModulationSettingsProxy(*Settings, *this);
+				}
+			}
+		});
+	}
+#endif // WITH_EDITOR
+
+	void FAudioModulationSystem::OnInitSound(ISoundModulatable& InSound, const USoundModulationPluginSourceSettingsBase& InSettings)
+	{
+		check(IsInAudioThread());
+
+		const uint32 SoundId = InSound.GetObjectId();
+		if (!SoundSettings.Contains(SoundId))
+		{
+			const USoundModulationSettings* Settings = CastChecked<USoundModulationSettings>(&InSettings);
+			SoundSettings.Add(SoundId, FModulationSettingsProxy(*Settings, *this));
+		}
+	}
+
+	void FAudioModulationSystem::OnInitSource(const uint32 InSourceId, const FName& AudioComponentUserId, const uint32 NumChannels, const USoundModulationPluginSourceSettingsBase& InSettings)
+	{
+		check(IsInAudioThread());
+
+		if (const USoundModulationSettings* Settings = CastChecked<USoundModulationSettings>(&InSettings))
+		{
+			SourceSettings[InSourceId] = FModulationSettingsProxy(*Settings, *this);
+		}
+	}
+
+	void FAudioModulationSystem::OnReleaseSource(const uint32 InSourceId)
+	{
+		check(IsInAudioThread());
+
+		SourceSettings[InSourceId] = FModulationSettingsProxy();
+	}
+
+	void FAudioModulationSystem::OnReleaseSound(ISoundModulatable& InSound)
+	{
+		check(IsInAudioThread());
+		check(InSound.GetObjectId() != INDEX_NONE);
+
+		if (InSound.GetPlayCount() == 0)
+		{
+			SoundSettings.Remove(InSound.GetObjectId());
+		}
+	}
+
+#if !UE_BUILD_SHIPPING
+	bool FAudioModulationSystem::OnPostHelp(FCommonViewportClient* ViewportClient, const TCHAR* Stream)
+	{
+		return ViewportClient ? Debugger.OnPostHelp(*ViewportClient, Stream) : true;
+	}
+
+	int32 FAudioModulationSystem::OnRenderStat(FViewport* Viewport, FCanvas* Canvas, int32 X, int32 Y, const UFont& Font, const FVector* ViewLocation, const FRotator* ViewRotation)
+	{
+		return Canvas ? Debugger.OnRenderStat(*Canvas, X, Y, Font) : Y;
+	}
+
+	bool FAudioModulationSystem::OnToggleStat(FCommonViewportClient* ViewportClient, const TCHAR* Stream)
+	{
+		return ViewportClient ? Debugger.OnToggleStat(*ViewportClient, Stream) : true;
+	}
+#endif // !UE_BUILD_SHIPPING
+
+	void FAudioModulationSystem::ActivateBus(const USoundControlBusBase& InBus)
+	{
+		const TWeakObjectPtr<const USoundControlBusBase>BusPtr = &InBus;
+		RunCommandOnAudioThread([this, BusPtr]()
+		{
+			if (BusPtr.IsValid())
+			{
+				auto OnCreate = [this, BusPtr](FControlBusProxy& NewProxy)
+				{
+					NewProxy.InitLFOs(*BusPtr.Get());
+				};
+
+				FBusHandle BusHandle = FBusHandle::Create(*BusPtr.Get(), RefProxies.Buses, *this, OnCreate);
+				ManuallyActivatedBuses.Add(MoveTemp(BusHandle));
+			}
+		});
+
+	}
+
+	void FAudioModulationSystem::ActivateBusMix(const USoundControlBusMix& InBusMix)
+	{
+		const TWeakObjectPtr<const USoundControlBusMix>BusMixPtr = &InBusMix;
+		RunCommandOnAudioThread([this, BusMixPtr]()
+		{
+			if (BusMixPtr.IsValid())
+			{
+				FBusMixHandle BusMixHandle = FBusMixHandle::Get(*BusMixPtr.Get(), RefProxies.BusMixes);
+				if (BusMixHandle.IsValid())
+				{
+					BusMixHandle.FindProxy().SetEnabled(*BusMixPtr.Get());
+				}
+				else
+				{
+					BusMixHandle = FBusMixHandle::Create(*BusMixPtr.Get(), RefProxies.BusMixes, *this);
+				}
+
+				ManuallyActivatedBusMixes.Add(MoveTemp(BusMixHandle));
+			}
+		});
+
+	}
+
+	void FAudioModulationSystem::ActivateLFO(const USoundBusModulatorLFO& InLFO)
+	{
+		const TWeakObjectPtr<const USoundBusModulatorLFO>LFOPtr = &InLFO;
+		RunCommandOnAudioThread([this, LFOPtr]()
+		{
+			if (LFOPtr.IsValid())
+			{
+				FLFOHandle LFOHandle = FLFOHandle::Create(*LFOPtr.Get(), RefProxies.LFOs, *this);
+				ManuallyActivatedLFOs.Add(MoveTemp(LFOHandle));
+			}
+		});
+	}
+
+	bool FAudioModulationSystem::CalculateModulationValue(FModulationPatchProxy& OutProxy, float& OutValue) const
+	{
+		if (OutProxy.IsBypassed())
+		{
+			return false;
+		}
+
+		const float InitValue = OutValue;
+		OutValue = OutProxy.Update();
+
+		return !FMath::IsNearlyEqual(InitValue, OutValue);
+	}
+
+	float FAudioModulationSystem::CalculateInitialVolume(const USoundModulationPluginSourceSettingsBase& SettingsBase)
+	{
+		check(IsInAudioThread());
+
+		const USoundModulationSettings* Settings = CastChecked<USoundModulationSettings>(&SettingsBase);
+		FModulationPatchProxy VolumePatch(Settings->Volume, *this);
+
+		return VolumePatch.Update();
+	}
+
+	void FAudioModulationSystem::DeactivateBus(const USoundControlBusBase& InBus)
+	{
+		const TWeakObjectPtr<const USoundControlBusBase>BusPtr = &InBus;
+		RunCommandOnAudioThread([this, BusPtr]()
+		{
+			if (BusPtr.IsValid())
+			{
+				const USoundControlBusBase& Bus = *BusPtr.Get();
+				FBusHandle BusHandle = FBusHandle::Create(Bus, RefProxies.Buses, *this);
+				ManuallyActivatedBuses.Remove(MoveTemp(BusHandle));
+			}
+		});
+	}
+
+	void FAudioModulationSystem::DeactivateBusMix(const USoundControlBusMix& InBusMix)
+	{
+		const TWeakObjectPtr<const USoundControlBusMix> BusMixPtr = &InBusMix;
+		RunCommandOnAudioThread([this, BusMixPtr]()
+		{
+			if (BusMixPtr.IsValid())
+			{
+				const USoundControlBusMix& BusMix = *BusMixPtr.Get();
+				FBusMixHandle BusMixHandle = FBusMixHandle::Get(BusMix, RefProxies.BusMixes);
+				if (BusMixHandle.IsValid())
+				{
+					FModulatorBusMixProxy& MixProxy = BusMixHandle.FindProxy();
+					MixProxy.SetStopping();
+				}
+			}
+		});
+	}
+
+	void FAudioModulationSystem::DeactivateLFO(const USoundBusModulatorLFO& InLFO)
+	{
+		const TWeakObjectPtr<const USoundBusModulatorLFO>LFOPtr = &InLFO;
+		RunCommandOnAudioThread([this, LFOPtr]()
+		{
+			if (LFOPtr.IsValid())
+			{
+				const USoundBusModulatorLFO& LFO = *LFOPtr.Get();
+				FLFOHandle LFOHandle = FLFOHandle::Create(LFO, RefProxies.LFOs, *this);
+				ManuallyActivatedLFOs.Remove(MoveTemp(LFOHandle));
+			}
+		});
+	}
+
+	bool FAudioModulationSystem::IsBusActive(const FBusId InBusId) const
+	{
+		check(IsInAudioThread());
+
+		return RefProxies.Buses.Contains(InBusId);
+	}
+
+	bool FAudioModulationSystem::ProcessControls(const uint32 InSourceId, FSoundModulationControls& OutControls)
+	{
+		check(IsInAudioThread());
+
+		bool bControlsUpdated = false;
+
+		FModulationSettingsProxy& Settings = SourceSettings[InSourceId];
+
+		if (Settings.Volume.IsBypassed())
+		{
+			OutControls.Volume = 1.0f;
+		}
+		else
+		{
+			bControlsUpdated |= CalculateModulationValue(Settings.Volume, OutControls.Volume);
+		}
+
+		if (Settings.Pitch.IsBypassed())
+		{
+			OutControls.Pitch = 1.0f;
+		}
+		else
+		{
+			bControlsUpdated |= CalculateModulationValue(Settings.Pitch, OutControls.Pitch);
+		}
+
+		if (Settings.Highpass.IsBypassed())
+		{
+			OutControls.Highpass = MIN_FILTER_FREQUENCY;
+		}
+		else
+		{
+			bControlsUpdated |= CalculateModulationValue(Settings.Highpass, OutControls.Highpass);
+		}
+
+		if (Settings.Lowpass.IsBypassed())
+		{
+			OutControls.Lowpass = MAX_FILTER_FREQUENCY;
+		}
+		else
+		{
+			bControlsUpdated |= CalculateModulationValue(Settings.Lowpass, OutControls.Lowpass);
+		}
+
+		return bControlsUpdated;
+	}
+
+	void FAudioModulationSystem::ProcessModulators(float Elapsed)
+	{
+		check(IsInAudioThread());
+
+		// Pump command queue
+		TUniqueFunction<void()> Command;
+		while (AudioThreadCommandQueue.Dequeue(Command))
+		{
+			Command();
+		}
+
+		// Update LFOs (prior to bus mixing to avoid single-frame latency)
+		for (TPair<FLFOId, FModulatorLFOProxy>& Pair : RefProxies.LFOs)
+		{
+			Pair.Value.Update(Elapsed);
+		}
+
+		// Reset buses & refresh cached LFO
+		for (TPair<FBusId, FControlBusProxy>& Pair : RefProxies.Buses)
+		{
+			Pair.Value.Reset();
+			Pair.Value.MixLFO();
+		}
+
+		// Update mix values and apply to prescribed buses.
+		// Track bus mixes ready to remove
+		TSet<FBusMixId> StoppedMixIds;
+		for (TPair<FBusMixId, FModulatorBusMixProxy>& Pair : RefProxies.BusMixes)
+		{
+			const FModulatorBusMixProxy::EStatus LastStatus = Pair.Value.GetStatus();
+			Pair.Value.Update(Elapsed, RefProxies.Buses);
+			const FModulatorBusMixProxy::EStatus CurrentStatus = Pair.Value.GetStatus();
+
+			switch (CurrentStatus)
+			{
+				case FModulatorBusMixProxy::EStatus::Enabled:
+				case FModulatorBusMixProxy::EStatus::Stopping:
+				break;
+
+				case FModulatorBusMixProxy::EStatus::Stopped:
+				{
+					if (LastStatus != CurrentStatus)
+					{
+						UE_LOG(LogAudioModulation, Log, TEXT("Audio modulation mix '%s' stopped."), *Pair.Value.GetName());
+					}
+					StoppedMixIds.Add(Pair.Key);
+				}
+				break;
+
+				default:
+				{
+					checkf(false, TEXT("Invalid or unsupported BusMix EStatus state advancement."));
+				}
+				break;
+			}
+		}
+
+		// Destroy mixes that have stopped (must be done outside mix update
+		// loop above to avoid destroying while iterating, which can occur
+		// when update moves bus mix from 'stopping' status to 'stopped')
+		for (const FBusMixId& MixId : StoppedMixIds)
+		{
+			FBusMixHandle MixHandle = FBusMixHandle::Get(MixId, RefProxies.BusMixes);
+
+			// Expected to be valid given the fact that the proxy is available in the prior loop
+			check(MixHandle.IsValid());
+
+			// Expected to only have two references (one for transient 'MixHandle' and one in
+			// ManuallyActivated set). Nothing else should be keeping mixes active.
+			check(MixHandle.FindProxy().GetRefCount() == 2);
+
+			ManuallyActivatedBusMixes.Remove(MoveTemp(MixHandle));
+		}
+
+		// Send updated patch/bus values to render thread
+		Audio::FControlModulatorValueMap UpdatedModValueMap;
+		for (TPair<FPatchId, FModulationPatchRefProxy>& Pair : RefProxies.Patches)
+		{
+			FModulationPatchRefProxy& PatchProxy = Pair.Value;
+			if (!PatchProxy.IsBypassed())
+			{
+				const float NewValue = PatchProxy.Update();
+				UpdatedModValueMap.Add(Pair.Key, NewValue);
+			}
+		}
+
+		for (TPair<FBusId, FControlBusProxy>& Pair : RefProxies.Buses)
+		{
+			FControlBusProxy& BusProxy = Pair.Value;
+			{
+				if (!BusProxy.IsBypassed())
+				{
+					UpdatedModValueMap.Add(Pair.Key, BusProxy.GetValue());
+				}
+			}
+		}
+
+		for (TPair<FLFOId, FModulatorLFOProxy>& Pair : RefProxies.LFOs)
+		{
+			FModulatorLFOProxy& LFOProxy = Pair.Value;
+			{
+				if (!LFOProxy.IsBypassed())
+				{
+					UpdatedModValueMap.Add(Pair.Key, LFOProxy.GetValue());
+				}
+			}
+		}
+
+		RunCommandOnAudioRenderThread([this, UpdatedModValueMap]()
+		{
+			ModValues_RenderThread = UpdatedModValueMap;
+		});
+
+		// Log stats
+		SET_DWORD_STAT(STAT_AudioModulationBusCount, RefProxies.Buses.Num());
+		SET_DWORD_STAT(STAT_AudioModulationMixCount, RefProxies.BusMixes.Num());
+		SET_DWORD_STAT(STAT_AudioModulationLFOCount, RefProxies.LFOs.Num());
+		SET_DWORD_STAT(STAT_AudioModulationPatchCount, RefProxies.Patches.Num());
+
+#if !UE_BUILD_SHIPPING
+		Debugger.UpdateDebugData(RefProxies);
+#endif // !UE_BUILD_SHIPPING
+	}
+
+	void FAudioModulationSystem::SaveMixToProfile(const USoundControlBusMix& InBusMix, const int32 InProfileIndex)
+	{
+		check(IsInGameThread());
+
+		const TWeakObjectPtr<const USoundControlBusMix> MixToSerialize = &InBusMix;
+		RunCommandOnAudioThread([this, MixToSerialize, InProfileIndex]()
+		{
+			if (!MixToSerialize.IsValid())
+			{
+				return;
+			}
+
+			const FBusMixId MixId = static_cast<FBusMixId>(MixToSerialize->GetUniqueID());
+			const FString   MixName = MixToSerialize->GetName();
+
+			FBusMixHandle MixHandle = FBusMixHandle::Get(MixId, RefProxies.BusMixes);
+			const bool bIsActive = MixHandle.IsValid();
+			if (!MixHandle.IsValid())
+			{
+				UE_LOG(LogAudioModulation, Display, TEXT("Mix '%s' is inactive, saving default object to profile '%i'."), *MixName, InProfileIndex);
+				AsyncTask(ENamedThreads::GameThread, [this, MixToSerialize, InProfileIndex]()
+				{
+					AudioModulation::FProfileSerializer::Serialize(*MixToSerialize.Get(), InProfileIndex);
+				});
+				return;
+			}
+
+			UE_LOG(LogAudioModulation, Display, TEXT("Mix '%s' is active, saving current mix proxy state to profile '%i'."), *MixName, InProfileIndex);
+			AudioModulation::FModulatorBusMixProxy& MixProxy = MixHandle.FindProxy();
+			TMap<FBusId, FSoundModulationValue> PassedChannelInfo;
+			for (TPair<FBusId, FModulatorBusMixChannelProxy>& Pair : MixProxy.Channels)
+			{
+				FModulatorBusMixChannelProxy& Channel = Pair.Value;
+				PassedChannelInfo.Add(Pair.Key, Channel.Value);
+			}
+
+			AsyncTask(ENamedThreads::GameThread, [this, PassedChannelInfo, MixToSerialize, InProfileIndex]()
+			{
+				if (!MixToSerialize.IsValid())
+				{
+					return;
+				}
+						
+				TMap<FBusId, FSoundModulationValue> ChannelInfo = PassedChannelInfo;
+				USoundControlBusMix* TempMix = NewObject<USoundControlBusMix>(GetTransientPackage(), *FGuid().ToString(EGuidFormats::Short));
+
+				// Buses on proxy may differ than those on uobject definition, so iterate and find by cached ids
+				// and add to temp mix to be serialized.
+				for (TObjectIterator<USoundControlBusBase> Itr; Itr; ++Itr)
+				{
+					if (USoundControlBusBase* Bus = *Itr)
+					{
+						FBusId ItrBusId = static_cast<FBusId>(Bus->GetUniqueID());
+						if (FSoundModulationValue* Value = ChannelInfo.Find(ItrBusId))
+						{
+							FSoundControlBusMixChannel BusMixChannel;
+							BusMixChannel.Bus = Bus;
+							BusMixChannel.Value = *Value;
+							TempMix->Channels.Add(MoveTemp(BusMixChannel));
+						}
+					}
+				}
+
+				const FString   MixPath = MixToSerialize->GetPathName();
+				AudioModulation::FProfileSerializer::Serialize(*TempMix, InProfileIndex, &MixPath);
+			});
+		});
+	}
+
+	TArray<FSoundControlBusMixChannel> FAudioModulationSystem::LoadMixFromProfile(const int32 InProfileIndex, USoundControlBusMix& OutBusMix)
+	{
+		const FString TempName = FGuid::NewGuid().ToString(EGuidFormats::Short);
+		if (USoundControlBusMix* TempMix = NewObject<USoundControlBusMix>(GetTransientPackage(), *TempName))
+		{
+			const FString MixPath = OutBusMix.GetPathName();
+			AudioModulation::FProfileSerializer::Deserialize(InProfileIndex, *TempMix, &MixPath);
+			UpdateMix(TempMix->Channels, OutBusMix);
+			return TempMix->Channels;
+		}
+
+		return TArray<FSoundControlBusMixChannel>();
+	}
+
+	void FAudioModulationSystem::RunCommandOnAudioThread(TUniqueFunction<void()> Cmd)
+	{
+		if (IsInAudioThread())
+		{
+			Cmd();
+		}
+		else
+		{
+			AudioThreadCommandQueue.Enqueue(MoveTemp(Cmd));
+		}
+	}
+
+	void FAudioModulationSystem::RunCommandOnAudioRenderThread(TUniqueFunction<void()> Cmd)
+	{
+		RenderThreadCommandQueue.Enqueue(MoveTemp(Cmd));
+	}
+
+	template <typename THandleType, typename TModType, typename TMapType>
+	bool FAudioModulationSystem::RegisterModulator(FAudioModulationSystem* InSystem, uint32 InParentId, const USoundModulatorBase& InModulatorBase, TMapType& ProxyMap, TMap<THandleType, TArray<uint32>>& ModMap)
+	{
+		if (const TModType* Mod = Cast<TModType>(&InModulatorBase))
+		{
+			const TWeakObjectPtr<const TModType> ModPtr = Mod;
+			TMapType* PassedProxyMap = &ProxyMap;
+			TMap<THandleType, TArray<uint32>>* PassedModMap = &ModMap;
+			RunCommandOnAudioThread([InSystem, ModPtr, InParentId, PassedProxyMap, PassedModMap]()
+			{
+				check(PassedProxyMap);
+				check(PassedModMap);
+					
+				if (ModPtr.IsValid())
+				{
+					THandleType Handle = THandleType::Create(*ModPtr.Get(), *PassedProxyMap, *InSystem);
+					PassedModMap->FindOrAdd(Handle).Add(InParentId);
+				}
+			});
+			return true;
+		}
+
+		return false;
+	}
+
+	bool FAudioModulationSystem::RegisterModulator(uint32 InParentId, const USoundModulatorBase& InModulatorBase)
+	{
+		if (RegisterModulator<FPatchHandle, USoundModulationPatch, FPatchProxyMap>(this, InParentId, InModulatorBase, RefProxies.Patches, RefModulators.PatchMap))
+		{
+			return true;
+		}
+			
+		if (RegisterModulator<FBusHandle, USoundControlBusBase, FBusProxyMap>(this, InParentId, InModulatorBase, RefProxies.Buses, RefModulators.BusMap))
+		{
+			return true;
+		}
+
+		if (RegisterModulator<FLFOHandle, USoundBusModulatorLFO, FLFOProxyMap>(this, InParentId, InModulatorBase, RefProxies.LFOs, RefModulators.LFOMap))
+		{
+			return true;
+		}
+
+		UE_LOG(LogAudioModulation, Warning, TEXT("Modulator type  of '%s' unsupported by generic control modulation."), *InModulatorBase.GetName());
+		return false;
+	}
+
+	bool FAudioModulationSystem::RegisterModulator(uint32 InParentId, Audio::FModulatorId InModulatorId)
+	{
+		RunCommandOnAudioThread([this, InParentId, InModulatorId]()
+		{
+			FPatchHandle PatchHandle = FPatchHandle::Get(static_cast<FPatchId>(InModulatorId), RefProxies.Patches);
+			if (PatchHandle.IsValid())
+			{
+				if (TArray<uint32>* RefObjectIds = RefModulators.PatchMap.Find(PatchHandle))
+				{
+					RefObjectIds->Add(InParentId);
+				}
+			}
+
+			FBusHandle BusHandle = FBusHandle::Get(static_cast<FBusId>(InModulatorId), RefProxies.Buses);
+			if (BusHandle.IsValid())
+			{
+				if (TArray<uint32>* RefObjectIds = RefModulators.BusMap.Find(BusHandle))
+				{
+					RefObjectIds->Add(InParentId);
+				}
+			}
+
+			FLFOHandle LFOHandle = FLFOHandle::Get(static_cast<FLFOId>(InModulatorId), RefProxies.LFOs);
+			if (LFOHandle.IsValid())
+			{
+				if (TArray<uint32>* RefObjectIds = RefModulators.LFOMap.Find(LFOHandle))
+				{
+					RefObjectIds->Add(InParentId);
+				}
+			}
+		});
+
+		return true;
+	}
+
+	bool FAudioModulationSystem::GetModulatorValue(const Audio::FModulatorHandle& InModulatorHandle, float& OutValue)
+	{
+		if (float* NewValue = ModValues_RenderThread.Find(InModulatorHandle.GetId()))
+		{
+			OutValue = *NewValue;
+			return true;
+		}
+
+		return false;
+	}
+
+	void FAudioModulationSystem::OnBeginAudioRenderThreadUpdate()
+	{
+		// Execute lambda functions passed from audio thread
+		// to update publicly available modulation values.
+		TUniqueFunction<void()> Command;
+		while (RenderThreadCommandQueue.Dequeue(Command))
+		{
+			Command();
+		}
+	}
+
+	template <typename THandleType>
+	bool FAudioModulationSystem::UnregisterModulator(THandleType PatchHandle, TMap<THandleType, TArray<uint32>>& HandleMap, const uint32 ParentId)
+	{
+		if (!PatchHandle.IsValid())
+		{
+			return false;
+		}
+
+		if (TArray<uint32>* ObjectIds = HandleMap.Find(PatchHandle))
+		{
+			for (int32 i = 0; i < ObjectIds->Num(); ++i)
+			{
+				const uint32 ObjectId = (*ObjectIds)[i];
+				if (ObjectId == ParentId)
+				{
+					ObjectIds->RemoveAtSwap(i, 1, false /* bAllowShrinking */);
+					if (ObjectIds->Num() == 0)
+					{
+						HandleMap.Remove(PatchHandle);
+					}
+					return true;
+				}
+			}
+		}
+
+		return false;
+	}
+
+	void FAudioModulationSystem::UnregisterModulator(const Audio::FModulatorHandle& InHandle)
+	{
+		const Audio::FModulatorId ModId = InHandle.GetId();
+		RunCommandOnAudioRenderThread([this, ModId]()
+		{
+			ModValues_RenderThread.Remove(ModId);
+		});
+
+		RunCommandOnAudioThread([this, ModId, ParentId = InHandle.GetParentId()]()
+		{
+			FPatchHandle PatchHandle = FPatchHandle::Get(static_cast<FPatchId>(ModId), RefProxies.Patches);
+			if (UnregisterModulator<FPatchHandle>(PatchHandle, RefModulators.PatchMap, ParentId))
+			{
+				return;
+			}
+
+			FBusHandle BusHandle = FBusHandle::Get(static_cast<FBusId>(ModId), RefProxies.Buses);
+			if (UnregisterModulator<FBusHandle>(BusHandle, RefModulators.BusMap, ParentId))
+			{
+				return;
+			}
+
+			FLFOHandle LFOHandle = FLFOHandle::Get(static_cast<FLFOId>(ModId), RefProxies.LFOs);
+			if (UnregisterModulator<FLFOHandle>(LFOHandle, RefModulators.LFOMap, ParentId))
+			{
+				return;
+			}
+		});
+	}
+
+	void FAudioModulationSystem::UpdateMix(const TArray<FSoundControlBusMixChannel>& InChannels, USoundControlBusMix& InOutMix, bool bInUpdateObject)
+	{
+		if (bInUpdateObject)
+		{
+			TMap<uint32, const FSoundControlBusMixChannel*> UpdatedChannelBusses;
+			for (const FSoundControlBusMixChannel& Channel : InChannels)
+			{
+				if (Channel.Bus)
+				{
+					UpdatedChannelBusses.Add(Channel.Bus->GetUniqueID(), &Channel);
+				}
+			}
+
+			bool bMarkDirty = false;
+			for (FSoundControlBusMixChannel& Channel : InOutMix.Channels)
+			{
+				if (!Channel.Bus)
+				{
+					continue;
+				}
+
+				if (const FSoundControlBusMixChannel* BusChannel = UpdatedChannelBusses.FindRef(Channel.Bus->GetUniqueID()))
+				{
+					Channel = *BusChannel;
+					bMarkDirty = true;
+				}
+			}
+			InOutMix.MarkPackageDirty();
+		}
+
+		const FBusMixId MixId = static_cast<FBusMixId>(InOutMix.GetUniqueID());
+		RunCommandOnAudioThread([this, MixId, InChannels]()
+		{
+			if (FModulatorBusMixProxy* BusMixes = RefProxies.BusMixes.Find(MixId))
+			{
+				BusMixes->SetMix(InChannels);
+			}
+		});
+	}
+
+	void FAudioModulationSystem::UpdateMixByFilter(
+		const FString&							 InAddressFilter,
+		const TSubclassOf<USoundControlBusBase>& InBusClass,
+		const FSoundModulationValue&			 InValue,
+		USoundControlBusMix&					 InOutMix,
+		bool									 bInUpdateObject)
+	{
+		const uint32 ClassId = InBusClass ? InBusClass->GetUniqueID() : USoundControlBusBase::StaticClass()->GetUniqueID();
+
+		if (bInUpdateObject)
+		{
+			bool bMarkDirty = false;
+			static const uint32 BaseBusClassId = USoundControlBusBase::StaticClass()->GetUniqueID();
+			for (FSoundControlBusMixChannel& Channel : InOutMix.Channels)
+			{
+				if (!Channel.Bus)
+				{
+					continue;
+				}
+
+				if (ClassId != BaseBusClassId && Channel.Bus->GetClass()->GetUniqueID() != ClassId)
+				{
+					continue;
+				}
+
+				if (!FAudioAddressPattern::PartsMatch(InAddressFilter, Channel.Bus->Address))
+				{
+					continue;
+				}
+
+				Channel.Value.TargetValue = InValue.TargetValue;
+
+				if (InValue.AttackTime >= 0.0f)
+				{
+					Channel.Value.AttackTime = InValue.AttackTime;
+				}
+
+				if (InValue.ReleaseTime >= 0.0f)
+				{
+					Channel.Value.ReleaseTime = InValue.ReleaseTime;
+				}
+				bMarkDirty = true;
+			}
+
+			if (bMarkDirty)
+			{
+				InOutMix.MarkPackageDirty();
+			}
+		}
+
+		const FString	AddressFilter = InAddressFilter;
+		const FBusMixId MixId = static_cast<FBusMixId>(InOutMix.GetUniqueID());
+		RunCommandOnAudioThread([this, ClassId, MixId, AddressFilter, InValue]()
+		{
+			if (FModulatorBusMixProxy* MixProxy = RefProxies.BusMixes.Find(MixId))
+			{
+				MixProxy->SetMixByFilter(AddressFilter, ClassId, InValue);
+			}
+		});
+	}
+
+	void FAudioModulationSystem::UpdateMix(const USoundControlBusMix& InMix)
+	{
+		const TWeakObjectPtr<const USoundControlBusMix> BusMixPtr = &InMix;
+		RunCommandOnAudioThread([this, BusMixPtr]()
+		{
+			if (BusMixPtr.IsValid())
+			{
+				const USoundControlBusMix& Mix = *BusMixPtr.Get();
+				FBusMixHandle BusMixHandle = FBusMixHandle::Get(Mix, RefProxies.BusMixes);
+				if (BusMixHandle.IsValid())
+				{
+					BusMixHandle.FindProxy() = Mix;
+				}
+				else
+				{
+					UE_LOG(LogAudioModulation, Verbose, TEXT("Update to '%s' Ignored: Control Bus Mix is inactive."), *Mix.GetName());
+				}
+			}
+		});
+	}
+
+	void FAudioModulationSystem::UpdateModulator(const USoundModulatorBase& InModulator)
+	{
+		if (const USoundBusModulatorLFO* InLFO = Cast<USoundBusModulatorLFO>(&InModulator))
+		{
+			const TWeakObjectPtr<const USoundBusModulatorLFO> LFOPtr = InLFO;
+			RunCommandOnAudioThread([this, LFOPtr]()
+			{
+				if (LFOPtr.IsValid())
+				{
+					const USoundBusModulatorLFO& LFO = *LFOPtr.Get();
+					FLFOHandle LFOHandle = FLFOHandle::Get(LFO, RefProxies.LFOs);
+					if (LFOHandle.IsValid())
+					{
+						LFOHandle.FindProxy() = LFO;
+					}
+					else
+					{
+						UE_LOG(LogAudioModulation, Verbose, TEXT("Update to '%s' Ignored: LFO is inactive."), *LFO.GetName());
+					}
+				}
+			});
+		}
+
+		if (const USoundControlBusBase* InBus = Cast<USoundControlBusBase>(&InModulator))
+		{
+			const TWeakObjectPtr<const USoundControlBusBase> BusPtr = InBus;
+			RunCommandOnAudioThread([this, BusPtr]()
+			{
+				if (BusPtr.IsValid())
+				{
+					const USoundControlBusBase* Bus = BusPtr.Get();
+					FBusHandle BusHandle = FBusHandle::Get(*Bus, RefProxies.Buses);
+					if (BusHandle.IsValid())
+					{
+						FControlBusProxy& Proxy = BusHandle.FindProxy();
+						Proxy = *Bus;
+						Proxy.InitLFOs(*Bus);
+					}
+					else
+					{
+						UE_LOG(LogAudioModulation, Verbose, TEXT("Update to '%s' Ignored: Control Bus is inactive."), *Bus->GetName());
+					}
+				}
+			});
+		}
+
+		if (const USoundControlBusMix* InMix = Cast<USoundControlBusMix>(&InModulator))
+		{
+			const TWeakObjectPtr<const USoundControlBusMix> BusMixPtr = InMix;
+			RunCommandOnAudioThread([this, BusMixPtr]()
+			{
+				if (BusMixPtr.IsValid())
+				{
+					const USoundControlBusMix& Mix = *BusMixPtr.Get();
+					FBusMixHandle BusMixHandle = FBusMixHandle::Get(Mix, RefProxies.BusMixes);
+					if (BusMixHandle.IsValid())
+					{
+						BusMixHandle.FindProxy() = Mix;
+					}
+					else
+					{
+						UE_LOG(LogAudioModulation, Verbose, TEXT("Update to '%s' Ignored: Control Bus Mix is inactive."), *Mix.GetName());
+					}
+				}
+			});
+		}
+
+		if (const USoundModulationPatch* InPatch = Cast<USoundModulationPatch>(&InModulator))
+		{
+			const TWeakObjectPtr<const USoundModulationPatch> PatchPtr = InPatch;
+			RunCommandOnAudioThread([this, PatchPtr]()
+			{
+				if (PatchPtr.IsValid())
+				{
+					const USoundModulationPatch& Patch = *PatchPtr.Get();
+					FPatchHandle PatchHandle = FPatchHandle::Get(Patch, RefProxies.Patches);
+					if (PatchHandle.IsValid())
+					{
+						FModulationPatchRefProxy& PatchProxy = PatchHandle.FindProxy();
+						PatchProxy = Patch;
+					}
+					else
+					{
+						UE_LOG(LogAudioModulation, Verbose, TEXT("Update to '%s' Ignored: Patch is inactive."), *Patch.GetName());
+					}
+				}
+			});
+		}
+	}
+} // namespace AudioModulation
+#endif // WITH_AUDIOMODULATION
diff --git a/Engine/Plugins/Runtime/AudioModulation/Source/AudioModulation/Private/AudioModulationSystem.h b/Engine/Plugins/Runtime/AudioModulation/Source/AudioModulation/Private/AudioModulationSystem.h
new file mode 100644
--- /dev/null
+++ b/Engine/Plugins/Runtime/AudioModulation/Source/AudioModulation/Private/AudioModulationSystem.h
@@ -0,0 +1,230 @@
+// Copyright Epic Games, Inc. All Rights Reserved.
+#pragma once
+
+#include "AudioModulation.h"
+#include "Containers/Queue.h"
+#include "CoreMinimal.h"
+#include "IAudioModulation.h"
+#include "SoundControlBus.h"
+#include "SoundControlBusMix.h"
+#include "SoundControlBusMixProxy.h"
+#include "SoundControlBusProxy.h"
+#include "SoundModulationPatchProxy.h"
+#include "SoundModulationProxy.h"
+#include "SoundModulationSettings.h"
+#include "SoundModulationSettingsProxy.h"
+#include "SoundModulationValue.h"
+#include "SoundModulatorLFO.h"
+#include "SoundModulatorLFOProxy.h"
+
+
+#if WITH_AUDIOMODULATION
+
+#if !UE_BUILD_SHIPPING
+#include "AudioModulationDebugger.h"
+#endif // !UE_BUILD_SHIPPING
+
+namespace AudioModulation
+{
+	// Forward Declarations
+	class FControlBusProxy;
+	class FModulationInputProxy;
+	class FModulationPatchProxy;
+	class FModulationPatchRefProxy;
+	class FModulationSettingsProxy;
+	class FModulatorBusMixChannelProxy;
+
+	struct FReferencedProxies
+	{
+		FBusMixProxyMap BusMixes;
+		FBusProxyMap    Buses;
+		FLFOProxyMap    LFOs;
+		FPatchProxyMap	Patches;
+	};
+
+	struct FReferencedModulators
+	{
+		TMap<FPatchHandle, TArray<uint32>> PatchMap;
+		TMap<FBusHandle, TArray<uint32>> BusMap;
+		TMap<FLFOHandle, TArray<uint32>> LFOMap;
+	};
+
+	class FAudioModulationSystem
+	{
+	public:
+		FAudioModulationSystem();
+
+		void Initialize(const FAudioPluginInitializationParams& InitializationParams);
+
+#if WITH_EDITOR
+		void OnEditPluginSettings(const USoundModulationPluginSourceSettingsBase& Settings);
+#endif // WITH_EDITOR
+
+		void OnInitSound(ISoundModulatable& InSound, const USoundModulationPluginSourceSettingsBase& Settings);
+		void OnInitSource(const uint32 InSourceId, const FName& AudioComponentUserId, const uint32 NumChannels, const USoundModulationPluginSourceSettingsBase& Settings);
+		void OnReleaseSound(ISoundModulatable& InSound);
+		void OnReleaseSource(const uint32 InSourceId);
+
+		void ActivateBus(const USoundControlBusBase& InBus);
+		void ActivateBusMix(const USoundControlBusMix& InBusMix);
+		void ActivateLFO(const USoundBusModulatorLFO& InLFO);
+
+		float CalculateInitialVolume(const USoundModulationPluginSourceSettingsBase& InSettingsBase);
+
+		/**
+		 * Deactivates respectively typed (i.e. BusMix, Bus, etc.) object proxy if no longer referenced.
+		 * If still referenced, will wait until references are finished before destroying.
+		 */
+		void DeactivateBus(const USoundControlBusBase& InBus);
+		void DeactivateBusMix(const USoundControlBusMix& InBusMix);
+		void DeactivateLFO(const USoundBusModulatorLFO& InLFO);
+
+		bool IsBusActive(const FBusId InBusId) const;
+
+		bool ProcessControls(const uint32 InSourceId, FSoundModulationControls& OutControls);
+		void ProcessModulators(const float Elapsed);
+
+		bool RegisterModulator(uint32 InParentId, const USoundModulatorBase& InModulatorBase);
+		bool RegisterModulator(uint32 InParentId, Audio::FModulatorId InModulatorId);
+		bool GetModulatorValue(const Audio::FModulatorHandle& ModulatorHandle, float& OutValue);
+		void UnregisterModulator(const Audio::FModulatorHandle& InHandle);
+
+		/* Saves mix to .ini profile for fast iterative development that does not require re-cooking a mix */
+		void SaveMixToProfile(const USoundControlBusMix& InBusMix, const int32 InProfileIndex);
+
+		/* Loads mix from .ini profile for iterative development that does not require re-cooking a mix. Returns copy
+		 * of mix channel values saved in profile. */
+		TArray<FSoundControlBusMixChannel> LoadMixFromProfile(const int32 InProfileIndex, USoundControlBusMix& OutBusMix);
+
+		/*
+		 * Updates mix/mix by filter, modifying the mix instance if it is active. If bInUpdateObject is true,
+		 * updates UObject definition in addition to proxy.
+		 */
+		void UpdateMix(const TArray<FSoundControlBusMixChannel>& InChannels, USoundControlBusMix& InOutMix, bool bInUpdateObject = false);
+		void UpdateMixByFilter(const FString& InAddressFilter, const TSubclassOf<USoundControlBusBase>& InClassFilter, const FSoundModulationValue& InValue, USoundControlBusMix& InOutMix, bool bInUpdateObject = false);
+
+		/*
+		 * Commits any changes from a mix applied to a UObject definition to mix instance if active.
+		 */
+		void UpdateMix(const USoundControlBusMix& InMix);
+
+		/*
+		 * Commits any changes from a modulator type applied to a UObject definition
+		 * to modulator instance if active (i.e. Control Bus, Control Bus Modulator)
+		 */
+		void UpdateModulator(const USoundModulatorBase& InModulator);
+
+		/*
+		 * Run on the audio render thread prior to processing audio
+		 */
+		void OnBeginAudioRenderThreadUpdate();
+
+	private:
+		/* Calculates modulation value, storing it in the provided float reference and returns if value changed */
+		bool CalculateModulationValue(FModulationPatchProxy& OutProxy, float& OutValue) const;
+
+		void RunCommandOnAudioThread(TUniqueFunction<void()> Cmd);
+
+		void RunCommandOnAudioRenderThread(TUniqueFunction<void()> Cmd);
+
+		template <typename THandleType, typename TModType, typename TMapType>
+		bool RegisterModulator(FAudioModulationSystem* InSystem, uint32 InParentId, const USoundModulatorBase& InModulatorBase, TMapType& ProxyMap, TMap<THandleType, TArray<uint32>>& ModMap);
+
+		template <typename THandleType>
+		bool UnregisterModulator(THandleType PatchHandle, TMap<THandleType, TArray<uint32>>& HandleMap, const uint32 ParentId);
+
+		FReferencedProxies RefProxies;
+
+		TSet<FBusHandle>    ManuallyActivatedBuses;
+		TSet<FBusMixHandle> ManuallyActivatedBusMixes;
+		TSet<FLFOHandle>    ManuallyActivatedLFOs;
+
+		// Cache of SoundId to sound setting state at time of sound initialization.
+		TMap<uint32, FModulationSettingsProxy> SoundSettings;
+
+		// Cache of source data while sound is actively playing
+		TArray<FModulationSettingsProxy> SourceSettings;
+
+		TQueue<TUniqueFunction<void()>> AudioThreadCommandQueue;
+		TQueue<TUniqueFunction<void()>> RenderThreadCommandQueue;
+
+		// Collection of maps with modulator handles to referencing object ids used by externally managing objects
+		FReferencedModulators RefModulators;
+
+		// Map of active patch values used by the audio render thread;
+		Audio::FControlModulatorValueMap ModValues_RenderThread;
+#if !UE_BUILD_SHIPPING
+	public:
+		bool OnPostHelp(FCommonViewportClient* ViewportClient, const TCHAR* Stream);
+		int OnRenderStat(FViewport* Viewport, FCanvas* Canvas, int32 X, int32 Y, const UFont& Font, const FVector* ViewLocation, const FRotator* ViewRotation);
+		bool OnToggleStat(FCommonViewportClient* ViewportClient, const TCHAR* Stream);
+
+	private:
+		FAudioModulationDebugger Debugger;
+#endif // !UE_BUILD_SHIPPING
+
+		friend FControlBusProxy;
+		friend FModulationInputProxy;
+		friend FModulationPatchProxy;
+		friend FModulationPatchRefProxy;
+		friend FModulationSettingsProxy;
+		friend FModulatorBusMixChannelProxy;
+	};
+} // namespace AudioModulation
+
+#else // WITH_AUDIOMODULATION
+
+namespace AudioModulation
+{
+	// Null implementation for compiler
+	class FAudioModulationSystem
+	{
+	public:
+		void Initialize(const FAudioPluginInitializationParams& InitializationParams) { }
+
+#if WITH_EDITOR
+		void OnEditPluginSettings(const USoundModulationPluginSourceSettingsBase& Settings) { }
+#endif // WITH_EDITOR
+
+		void OnInitSound(ISoundModulatable& InSound, const USoundModulationPluginSourceSettingsBase& Settings) { }
+		void OnInitSource(const uint32 InSourceId, const FName& AudioComponentUserId, const uint32 NumChannels, const USoundModulationPluginSourceSettingsBase& Settings) { }
+		void OnReleaseSound(ISoundModulatable& InSound) { }
+		void OnReleaseSource(const uint32 InSourceId) { }
+
+#if !UE_BUILD_SHIPPING
+		bool OnPostHelp(FCommonViewportClient* ViewportClient, const TCHAR* Stream) { return false; }
+		int OnRenderStat(FViewport* Viewport, FCanvas* Canvas, int32 X, int32 Y, const UFont& Font, const FVector* ViewLocation, const FRotator* ViewRotation) { return Y; }
+		bool OnToggleStat(FCommonViewportClient* ViewportClient, const TCHAR* Stream) { return false; }
+#endif // !UE_BUILD_SHIPPING
+
+		void ActivateBus(const USoundControlBusBase& InBus) { }
+		void ActivateBusMix(const USoundControlBusMix& InBusMix) { }
+		void ActivateLFO(const USoundBusModulatorLFO& InLFO) { }
+
+		float CalculateInitialVolume(const USoundModulationPluginSourceSettingsBase& InSettingsBase) { return 1.0f; }
+
+		void DeactivateBus(const USoundControlBusBase& InBus) { }
+		void DeactivateBusMix(const USoundControlBusBase& InBusMix) { }
+		void DeactivateLFO(const USoundBusModulatorLFO& InLFO) { }
+
+		void SaveMixToProfile(const USoundControlBusMix& InBusMix, const int32 InProfileIndex) { }
+		void LoadMixFromProfile(const int32 InProfileIndex, USoundControlBusMix& OutBusMix) { }
+
+		void OnBeginAudioRenderThreadUpdate() override { }
+
+		bool ProcessControls(const uint32 InSourceId, FSoundModulationControls& OutControls) { return false; }
+		void ProcessModulators(const float Elapsed) { }
+
+		bool RegisterModulator(uint32 InParentId, const USoundModulatorBase& InModulatorBase) { return false; }
+		bool RegisterModulator(uint32 InParentId, Audio::FModulatorId InModulatorId) { return false; }
+		bool GetModulatorValue(const Audio::FModulatorHandle& ModulatorHandle, float& OutValue) { return false;	}
+		void UnregisterModulator(const Audio::FModulatorHandle& InHandle) { }
+
+		void UpdateMix(const USoundControlBusMix& InMix) { }
+		void UpdateMix(const TArray<FSoundControlBusMixChannel>& InChannels, USoundControlBusMix& InOutMix, bool bUpdateObject = false) { }
+		void UpdateMixByFilter(const FString& InAddressFilter, const TSubclassOf<USoundControlBusBase>& InClassFilter, const FSoundModulationValue& InValue, USoundControlBusMix& InOutMix, bool bUpdateObject = false) { }
+		void UpdateModulator(const USoundModulatorBase& InModulator) { }
+
+	};
+} // namespace AudioModulation
+#endif // WITH_AUDIOMODULATION
diff --git a/Engine/Plugins/Runtime/AudioModulation/Source/AudioModulation/Private/SoundControlBus.cpp b/Engine/Plugins/Runtime/AudioModulation/Source/AudioModulation/Private/SoundControlBus.cpp
--- a/Engine/Plugins/Runtime/AudioModulation/Source/AudioModulation/Private/SoundControlBus.cpp
+++ b/Engine/Plugins/Runtime/AudioModulation/Source/AudioModulation/Private/SoundControlBus.cpp
@@ -4,11 +4,11 @@
 #include "AudioDevice.h"
 #include "AudioDeviceManager.h"
 #include "AudioModulation.h"
-#include "AudioModulationInternal.h"
 #include "AudioModulationLogging.h"
+#include "AudioModulationSystem.h"
 #include "Engine/World.h"
+#include "SoundControlBusProxy.h"
 #include "SoundModulatorLFO.h"
-#include "SoundModulationProxy.h"
 
 USoundControlBusBase::USoundControlBusBase(const FObjectInitializer& ObjectInitializer)
 	: Super(ObjectInitializer)
@@ -33,9 +33,9 @@ void USoundControlBusBase::PostDuplicate(EDuplicateMode::Type DuplicateMode)
 	Super::PostDuplicate(DuplicateMode);
 }
 
-void USoundControlBusBase::PostEditChangeProperty(FPropertyChangedEvent& PropertyChangedEvent)
+void USoundControlBusBase::PostEditChangeProperty(FPropertyChangedEvent& InPropertyChangedEvent)
 {
-	if (FProperty* Property = PropertyChangedEvent.Property)
+	if (FProperty* Property = InPropertyChangedEvent.Property)
 	{
 		if (Property->GetFName() == GET_MEMBER_NAME_CHECKED(USoundControlBusBase, bOverrideAddress) && !bOverrideAddress)
 		{
@@ -49,9 +49,11 @@ void USoundControlBusBase::PostEditChangeProperty(FPropertyChangedEvent& Propert
 			Min = FMath::Min(Min, Max);
 			DefaultValue = FMath::Clamp(DefaultValue, Min, Max);
 		}
+
+		AudioModulation::OnEditModulator(InPropertyChangedEvent, *this);
 	}
 
-	Super::PostEditChangeProperty(PropertyChangedEvent);
+	Super::PostEditChangeProperty(InPropertyChangedEvent);
 }
 
 void USoundControlBusBase::PostInitProperties()
@@ -110,172 +112,10 @@ void USoundControlBusBase::BeginDestroy()
 			check(AudioDevice->IsModulationPluginEnabled());
 			if (IAudioModulation* ModulationInterface = AudioDevice->ModulationInterface.Get())
 			{
-				auto ModulationImpl = static_cast<AudioModulation::FAudioModulation*>(ModulationInterface)->GetImpl();
-				check(ModulationImpl);
-				ModulationImpl->DeactivateBus(*this);
+				auto ModSystem = static_cast<AudioModulation::FAudioModulation*>(ModulationInterface)->GetModulationSystem();
+				check(ModSystem);
+				ModSystem->DeactivateBus(*this);
 			}
 		}
 	}
 }
-
-namespace AudioModulation
-{
-	FControlBusProxy::FControlBusProxy()
-		: DefaultValue(0.0f)
-		, LFOValue(1.0f)
-		, MixValue(NAN)
-		, bBypass(false)
-		, Operator(ESoundModulatorOperator::Multiply)
-		, Range(0.0f, 1.0f)
-	{
-	}
-
-	FControlBusProxy::FControlBusProxy(const USoundControlBusBase& InBus, FAudioModulationImpl& InModulationImpl)
-		: TModulatorProxyRefType(InBus.GetName(), InBus.GetUniqueID(), InModulationImpl)
-	{
-		Init(InBus);
-	}
-
-	FControlBusProxy& FControlBusProxy::operator =(const USoundControlBusBase& InBus)
-	{
-		Init(InBus);
-		return *this;
-	}
-
-	float FControlBusProxy::GetDefaultValue() const
-	{
-		return DefaultValue;
-	}
-
-	const TArray<FLFOHandle>& FControlBusProxy::GetLFOHandles() const
-	{
-		return LFOHandles;
-	}
-
-	float FControlBusProxy::GetLFOValue() const
-	{
-		return LFOValue;
-	}
-
-	FVector2D FControlBusProxy::GetRange() const
-	{
-		return Range;
-	}
-
-	float FControlBusProxy::GetMixValue() const
-	{
-		return MixValue;
-	}
-
-	float FControlBusProxy::GetValue() const
-	{
-		const float DefaultMixed = Mix(DefaultValue);
-		return FMath::Clamp(DefaultMixed * LFOValue, Range.X, Range.Y);
-	}
-
-	void FControlBusProxy::Init(const USoundControlBusBase& InBus)
-	{
-		DefaultValue = InBus.DefaultValue;
-		LFOValue = 1.0f;
-		MixValue = NAN;
-		Operator = InBus.GetOperator();
-		Range = FVector2D(InBus.Min, InBus.Max);
-		if (InBus.Min > InBus.Max)
-		{
-			Range.X = InBus.Max;
-			Range.Y = InBus.Min;
-		}
-
-		DefaultValue = FMath::Clamp(DefaultValue, Range.X, Range.Y);
-		bBypass = InBus.bBypass;
-	}
-
-	void FControlBusProxy::InitLFOs(const USoundControlBusBase& InBus, FLFOProxyMap& OutActiveLFOs)
-	{
-		LFOHandles.Reset();
-
-		for (const USoundBusModulatorBase* Modulator : InBus.Modulators)
-		{
-			if (const USoundBusModulatorLFO* LFO = Cast<USoundBusModulatorLFO>(Modulator))
-			{
-				check(ModulationImpl);
-				LFOHandles.Add(FLFOHandle::Create(*LFO, OutActiveLFOs, *ModulationImpl));
-			}
-		}
-	}
-
-	bool FControlBusProxy::IsBypassed() const
-	{
-		return bBypass;
-	}
-
-	float FControlBusProxy::Mix(float ValueA) const
-	{
-		// If mix value is NaN, it is uninitialized (effectively, the parent bus is inactive)
-		// and therefore not mixable, so just return the second value.
-		if (FMath::IsNaN(MixValue))
-		{
-			return ValueA;
-		}
-		return Mix(MixValue, ValueA);
-	}
-		
-	float FControlBusProxy::Mix(float ValueA, float ValueB) const
-	{
-		switch (Operator)
-		{
-			case ESoundModulatorOperator::Min:
-			{
-				return FMath::Min(ValueA, ValueB);
-			}
-			break;
-
-			case ESoundModulatorOperator::Max:
-			{
-				return FMath::Max(ValueA, ValueB);
-			}
-			break;
-
-			case ESoundModulatorOperator::Multiply:
-			{
-				return ValueA * ValueB;
-			}
-			break;
-
-			default:
-			{
-				static_assert(static_cast<int32>(ESoundModulatorOperator::Count) == 3, "Possible missing ESoundModulationOperator case coverage");
-			}
-			break;
-		}
-
-		check(false);
-		return NAN;
-	}
-
-	void FControlBusProxy::MixIn(const float InValue)
-	{
-		MixValue = Mix(InValue);
-	}
-
-	void FControlBusProxy::MixLFO()
-	{
-		for (const FLFOHandle& Handle: LFOHandles)
-		{
-			if (Handle.IsValid())
-			{
-				const FModulatorLFOProxy& LFOProxy = Handle.FindProxy();
-				if (!LFOProxy.IsBypassed())
-				{
-					LFOValue *= LFOProxy.GetValue();
-				}
-			}
-		}
-	}
-
-	void FControlBusProxy::Reset()
-	{
-		LFOValue = 1.0f;
-		MixValue = NAN;
-	}
-} // namespace AudioModulation
diff --git a/Engine/Plugins/Runtime/AudioModulation/Source/AudioModulation/Private/SoundControlBusMix.cpp b/Engine/Plugins/Runtime/AudioModulation/Source/AudioModulation/Private/SoundControlBusMix.cpp
--- a/Engine/Plugins/Runtime/AudioModulation/Source/AudioModulation/Private/SoundControlBusMix.cpp
+++ b/Engine/Plugins/Runtime/AudioModulation/Source/AudioModulation/Private/SoundControlBusMix.cpp
@@ -4,11 +4,12 @@
 #include "Audio/AudioAddressPattern.h"
 #include "AudioDevice.h"
 #include "AudioModulation.h"
-#include "AudioModulationInternal.h"
 #include "AudioModulationLogging.h"
 #include "AudioModulationProfileSerializer.h"
 #include "AudioModulationStatics.h"
+#include "AudioModulationSystem.h"
 #include "Engine/World.h"
+#include "SoundControlBus.h"
 
 #if WITH_EDITOR
 #include "Framework/Notifications/NotificationManager.h"
@@ -48,9 +49,9 @@ void USoundControlBusMix::BeginDestroy()
 			check(AudioDevice->IsModulationPluginEnabled());
 			if (IAudioModulation* ModulationInterface = AudioDevice->ModulationInterface.Get())
 			{
-				auto ModulationImpl = static_cast<AudioModulation::FAudioModulation*>(ModulationInterface)->GetImpl();
-				check(ModulationImpl);
-				ModulationImpl->DeactivateBusMix(*this);
+				auto ModSystem = static_cast<AudioModulation::FAudioModulation*>(ModulationInterface)->GetModulationSystem();
+				check(ModSystem);
+				ModSystem->DeactivateBusMix(*this);
 			}
 		}
 	}
@@ -73,181 +74,55 @@ void USoundControlBusMix::LoadMixFromProfile()
 	}
 }
 
-void USoundControlBusMix::SaveMixToProfile()
+void USoundControlBusMix::PostEditChangeProperty(FPropertyChangedEvent& InPropertyChangedEvent)
 {
-	if (AudioModulation::FProfileSerializer::Serialize(*this, ProfileIndex))
-	{
-		{
-			FNotificationInfo Info(FText::Format(
-				LOCTEXT("SoundControlBusMix_SaveSucceeded", "'Control Bus Mix '{0}' profile {1} saved successfully."),
-					FText::FromName(GetFName()),
-					FText::AsNumber(ProfileIndex)
-			));
-			Info.bFireAndForget = true;
-			Info.ExpireDuration = 2.0f;
-			Info.bUseThrobber = true;
-			FSlateNotificationManager::Get().AddNotification(Info);
-		}
-	}
-}
-#endif // WITH_EDITOR
-
-namespace AudioModulation
-{
-	FModulatorBusMixChannelProxy::FModulatorBusMixChannelProxy(const FSoundControlBusMixChannel& Channel, FAudioModulationImpl& InModulationImpl)
-		: TModulatorProxyBase<FBusId>(Channel.Bus->GetName(), Channel.Bus->GetUniqueID())
-		, Address(Channel.Bus->Address)
-		, ClassId(Channel.Bus->GetClass()->GetUniqueID())
-		, Value(Channel.Value)
-		, BusHandle(FBusHandle::Create(*Channel.Bus, InModulationImpl.RefProxies.Buses, InModulationImpl))
-	{
-	}
-
-	FModulatorBusMixProxy::FModulatorBusMixProxy(const USoundControlBusMix& InBusMix, FAudioModulationImpl& InModulationImpl)
-		: TModulatorProxyRefType(InBusMix.GetName(), InBusMix.GetUniqueID(), InModulationImpl)
-	{
-		SetEnabled(InBusMix);
-	}
-
-	FModulatorBusMixProxy& FModulatorBusMixProxy::operator =(const USoundControlBusMix& InBusMix)
-	{
-		SetEnabled(InBusMix);
-
-		return *this;
-	}
-
-	FModulatorBusMixProxy::EStatus FModulatorBusMixProxy::GetStatus() const
-	{
-		return Status;
-	}
-
-	void FModulatorBusMixProxy::Reset()
+	if (FProperty* Property = InPropertyChangedEvent.Property)
 	{
-		Channels.Reset();
-	}
-
-	void FModulatorBusMixProxy::SetEnabled(const USoundControlBusMix& InBusMix)
-	{
-		FChannelMap CachedChannels = Channels;
-		Channels.Reset();
-
-		Status = EStatus::Enabled;
-		for (const FSoundControlBusMixChannel& Channel : InBusMix.Channels)
+		if (!GEngine)
 		{
-			if (Channel.Bus)
-			{
-				auto BusId = static_cast<const FBusId>(Channel.Bus->GetUniqueID());
-
-				check(ModulationImpl);
-				FModulatorBusMixChannelProxy ChannelProxy(Channel, *ModulationImpl);
-
-				if (const FModulatorBusMixChannelProxy* CachedChannel = CachedChannels.Find(BusId))
-				{
-					ChannelProxy.Value.SetCurrentValue(CachedChannel->Value.GetCurrentValue());
-				}
-
-				Channels.Emplace(BusId, ChannelProxy);
-			}
-			else
-			{
-				UE_LOG(LogAudioModulation, Warning,
-					TEXT("USoundControlBusMix '%s' has channel with no bus specified. "
-						"Mix instance initialized but channel ignored."),
-					*InBusMix.GetFullName());
-			}
+			return;
 		}
-	}
 
-	void FModulatorBusMixProxy::SetMix(const TArray<FSoundControlBusMixChannel>& InChannels)
-	{
-		for (const FSoundControlBusMixChannel& NewChannel : InChannels)
+		FAudioDeviceManager* DeviceManager = GEngine->GetAudioDeviceManager();
+		if (!DeviceManager)
 		{
-			if(NewChannel.Bus)
-			{
-				FBusId BusId = static_cast<FBusId>(NewChannel.Bus->GetUniqueID());
-				if (FModulatorBusMixChannelProxy* ChannelProxy = Channels.Find(BusId))
-				{
-					ChannelProxy->Value.TargetValue = NewChannel.Value.TargetValue;
-					ChannelProxy->Value.AttackTime = NewChannel.Value.AttackTime;
-					ChannelProxy->Value.ReleaseTime = NewChannel.Value.ReleaseTime;
-				}
-			}
+			return;
 		}
-	}
 
-	void FModulatorBusMixProxy::SetMixByFilter(const FString& InAddressFilter, uint32 InFilterClassId, const FSoundModulationValue& InValue)
-	{
-		static const uint32 BaseClassId = USoundControlBusBase::StaticClass()->GetUniqueID();
-
-		for (TPair<FBusId, FModulatorBusMixChannelProxy>& IdProxyPair : Channels)
+		TArray<FAudioDevice*> Devices = DeviceManager->GetAudioDevices();
+		for (FAudioDevice* Device : Devices)
 		{
-			FModulatorBusMixChannelProxy& ChannelProxy = IdProxyPair.Value;
-			if (InFilterClassId != BaseClassId && ChannelProxy.ClassId != InFilterClassId)
+			if (Device && Device->IsModulationPluginEnabled() && Device->ModulationInterface.IsValid())
 			{
-				continue;
-			}
+				auto ModulationInterface = static_cast<AudioModulation::FAudioModulation*>(Device->ModulationInterface.Get());
+				AudioModulation::FAudioModulationSystem* ModulationSystem = ModulationInterface->GetModulationSystem();
+				check(ModulationSystem);
 
-			if (!FAudioAddressPattern::PartsMatch(InAddressFilter, ChannelProxy.Address))
-			{
-				continue;
-			}
-
-			ChannelProxy.Value.TargetValue = InValue.TargetValue;
-
-			if (InValue.AttackTime >= 0.0f)
-			{
-				ChannelProxy.Value.AttackTime = InValue.AttackTime;
-			}
-
-			if (InValue.ReleaseTime >= 0.0f)
-			{
-				ChannelProxy.Value.ReleaseTime = InValue.ReleaseTime;
+				ModulationSystem->UpdateMix(*this);
 			}
 		}
 	}
 
-	void FModulatorBusMixProxy::SetStopping()
-	{
-		if (Status == EStatus::Enabled)
-		{
-			Status = EStatus::Stopping;
-		}
-	}
+	Super::PostEditChangeProperty(InPropertyChangedEvent);
+}
 
-	void FModulatorBusMixProxy::Update(float Elapsed, FBusProxyMap& ProxyMap)
+void USoundControlBusMix::SaveMixToProfile()
+{
+	if (AudioModulation::FProfileSerializer::Serialize(*this, ProfileIndex))
 	{
-		bool bRequestStop = true;
-		for (TPair<FBusId, FModulatorBusMixChannelProxy>& Channel : Channels)
 		{
-			FModulatorBusMixChannelProxy& ChannelProxy = Channel.Value;
-			FSoundModulationValue& MixChannelValue = ChannelProxy.Value;
-
-			if (FControlBusProxy* BusProxy = ProxyMap.Find(ChannelProxy.GetId()))
-			{
-				MixChannelValue.Update(Elapsed);
-
-				const float CurrentValue = MixChannelValue.GetCurrentValue();
-				if (Status == EStatus::Stopping)
-				{
-					MixChannelValue.TargetValue = BusProxy->GetDefaultValue();
-					if (!FMath::IsNearlyEqual(MixChannelValue.TargetValue, CurrentValue))
-					{
-						bRequestStop = false;
-					}
-				}
-				else
-				{
-					bRequestStop = false;
-				}
-				BusProxy->MixIn(CurrentValue);
-			}
-		}
-
-		if (bRequestStop)
-		{
-			Status = EStatus::Stopped;
+			FNotificationInfo Info(FText::Format(
+				LOCTEXT("SoundControlBusMix_SaveSucceeded", "'Control Bus Mix '{0}' profile {1} saved successfully."),
+					FText::FromName(GetFName()),
+					FText::AsNumber(ProfileIndex)
+			));
+			Info.bFireAndForget = true;
+			Info.ExpireDuration = 2.0f;
+			Info.bUseThrobber = true;
+			FSlateNotificationManager::Get().AddNotification(Info);
 		}
 	}
-} // namespace AudioModulation
+}
+#endif // WITH_EDITOR
 
 #undef LOCTEXT_NAMESPACE // AudioModulation
diff --git a/Engine/Plugins/Runtime/AudioModulation/Source/AudioModulation/Private/SoundControlBusMixProxy.cpp b/Engine/Plugins/Runtime/AudioModulation/Source/AudioModulation/Private/SoundControlBusMixProxy.cpp
new file mode 100644
--- /dev/null
+++ b/Engine/Plugins/Runtime/AudioModulation/Source/AudioModulation/Private/SoundControlBusMixProxy.cpp
@@ -0,0 +1,176 @@
+// Copyright Epic Games, Inc. All Rights Reserved.
+#include "SoundControlBusMixProxy.h"
+
+#include "Audio/AudioAddressPattern.h"
+#include "AudioDevice.h"
+#include "AudioModulation.h"
+#include "AudioModulationLogging.h"
+#include "AudioModulationProfileSerializer.h"
+#include "AudioModulationStatics.h"
+#include "AudioModulationSystem.h"
+#include "SoundControlBus.h"
+
+#define LOCTEXT_NAMESPACE "AudioModulation"
+
+
+namespace AudioModulation
+{
+	const FBusMixId InvalidBusMixId = INDEX_NONE;
+
+	FModulatorBusMixChannelProxy::FModulatorBusMixChannelProxy(const FSoundControlBusMixChannel& Channel, FAudioModulationSystem& InModSystem)
+		: TModulatorProxyBase<FBusId>(Channel.Bus->GetName(), Channel.Bus->GetUniqueID())
+		, Address(Channel.Bus->Address)
+		, ClassId(Channel.Bus->GetClass()->GetUniqueID())
+		, Value(Channel.Value)
+		, BusHandle(FBusHandle::Create(*Channel.Bus, InModSystem.RefProxies.Buses, InModSystem))
+	{
+	}
+
+	FModulatorBusMixProxy::FModulatorBusMixProxy(const USoundControlBusMix& InBusMix, FAudioModulationSystem& InModSystem)
+		: TModulatorProxyRefType(InBusMix.GetName(), InBusMix.GetUniqueID(), InModSystem)
+	{
+		SetEnabled(InBusMix);
+	}
+
+	FModulatorBusMixProxy& FModulatorBusMixProxy::operator =(const USoundControlBusMix& InBusMix)
+	{
+		SetEnabled(InBusMix);
+
+		return *this;
+	}
+
+	FModulatorBusMixProxy::EStatus FModulatorBusMixProxy::GetStatus() const
+	{
+		return Status;
+	}
+
+	void FModulatorBusMixProxy::Reset()
+	{
+		Channels.Reset();
+	}
+
+	void FModulatorBusMixProxy::SetEnabled(const USoundControlBusMix& InBusMix)
+	{
+		FChannelMap CachedChannels = Channels;
+		Channels.Reset();
+
+		Status = EStatus::Enabled;
+		for (const FSoundControlBusMixChannel& Channel : InBusMix.Channels)
+		{
+			if (Channel.Bus)
+			{
+				auto BusId = static_cast<const FBusId>(Channel.Bus->GetUniqueID());
+
+				check(ModSystem);
+				FModulatorBusMixChannelProxy ChannelProxy(Channel, *ModSystem);
+
+				if (const FModulatorBusMixChannelProxy* CachedChannel = CachedChannels.Find(BusId))
+				{
+					ChannelProxy.Value.SetCurrentValue(CachedChannel->Value.GetCurrentValue());
+				}
+
+				Channels.Emplace(BusId, ChannelProxy);
+			}
+			else
+			{
+				UE_LOG(LogAudioModulation, Warning,
+					TEXT("USoundControlBusMix '%s' has channel with no bus specified. "
+						"Mix instance initialized but channel ignored."),
+					*InBusMix.GetFullName());
+			}
+		}
+	}
+
+	void FModulatorBusMixProxy::SetMix(const TArray<FSoundControlBusMixChannel>& InChannels)
+	{
+		for (const FSoundControlBusMixChannel& NewChannel : InChannels)
+		{
+			if(NewChannel.Bus)
+			{
+				FBusId BusId = static_cast<FBusId>(NewChannel.Bus->GetUniqueID());
+				if (FModulatorBusMixChannelProxy* ChannelProxy = Channels.Find(BusId))
+				{
+					ChannelProxy->Value.TargetValue = NewChannel.Value.TargetValue;
+					ChannelProxy->Value.AttackTime = NewChannel.Value.AttackTime;
+					ChannelProxy->Value.ReleaseTime = NewChannel.Value.ReleaseTime;
+				}
+			}
+		}
+	}
+
+	void FModulatorBusMixProxy::SetMixByFilter(const FString& InAddressFilter, uint32 InFilterClassId, const FSoundModulationValue& InValue)
+	{
+		static const uint32 BaseClassId = USoundControlBusBase::StaticClass()->GetUniqueID();
+
+		for (TPair<FBusId, FModulatorBusMixChannelProxy>& IdProxyPair : Channels)
+		{
+			FModulatorBusMixChannelProxy& ChannelProxy = IdProxyPair.Value;
+			if (InFilterClassId != BaseClassId && ChannelProxy.ClassId != InFilterClassId)
+			{
+				continue;
+			}
+
+			if (!FAudioAddressPattern::PartsMatch(InAddressFilter, ChannelProxy.Address))
+			{
+				continue;
+			}
+
+			ChannelProxy.Value.TargetValue = InValue.TargetValue;
+
+			if (InValue.AttackTime >= 0.0f)
+			{
+				ChannelProxy.Value.AttackTime = InValue.AttackTime;
+			}
+
+			if (InValue.ReleaseTime >= 0.0f)
+			{
+				ChannelProxy.Value.ReleaseTime = InValue.ReleaseTime;
+			}
+		}
+	}
+
+	void FModulatorBusMixProxy::SetStopping()
+	{
+		if (Status == EStatus::Enabled)
+		{
+			Status = EStatus::Stopping;
+		}
+	}
+
+	void FModulatorBusMixProxy::Update(float Elapsed, FBusProxyMap& ProxyMap)
+	{
+		bool bRequestStop = true;
+		for (TPair<FBusId, FModulatorBusMixChannelProxy>& Channel : Channels)
+		{
+			FModulatorBusMixChannelProxy& ChannelProxy = Channel.Value;
+			FSoundModulationValue& MixChannelValue = ChannelProxy.Value;
+
+			if (FControlBusProxy* BusProxy = ProxyMap.Find(ChannelProxy.GetId()))
+			{
+				MixChannelValue.Update(Elapsed);
+
+				const float CurrentValue = MixChannelValue.GetCurrentValue();
+				if (Status == EStatus::Stopping)
+				{
+					MixChannelValue.TargetValue = BusProxy->GetDefaultValue();
+					if (!FMath::IsNearlyEqual(MixChannelValue.TargetValue, CurrentValue))
+					{
+						bRequestStop = false;
+					}
+				}
+				else
+				{
+					bRequestStop = false;
+				}
+				BusProxy->MixIn(CurrentValue);
+			}
+		}
+
+		if (bRequestStop)
+		{
+			Status = EStatus::Stopped;
+		}
+	}
+} // namespace AudioModulation
+
+#undef LOCTEXT_NAMESPACE // AudioModulation
diff --git a/Engine/Plugins/Runtime/AudioModulation/Source/AudioModulation/Private/SoundControlBusMixProxy.h b/Engine/Plugins/Runtime/AudioModulation/Source/AudioModulation/Private/SoundControlBusMixProxy.h
new file mode 100644
--- /dev/null
+++ b/Engine/Plugins/Runtime/AudioModulation/Source/AudioModulation/Private/SoundControlBusMixProxy.h
@@ -0,0 +1,66 @@
+// Copyright Epic Games, Inc. All Rights Reserved.
+#pragma once
+
+#include "IAudioModulation.h"
+#include "SoundModulationPatchProxy.h"
+#include "SoundModulationProxy.h"
+#include "SoundModulationValue.h"
+#include "SoundModulatorLFOProxy.h"
+
+
+namespace AudioModulation
+{
+	// Forward Declarations
+	class FAudioModulationSystem;
+
+	using FBusMixId = uint32;
+	extern const FBusMixId InvalidBusMixId;
+
+	class FModulatorBusMixChannelProxy : public TModulatorProxyBase<FBusId>
+	{
+	public:
+
+		FModulatorBusMixChannelProxy(const FSoundControlBusMixChannel& Channel, FAudioModulationSystem& ModSystem);
+
+		FString Address;
+		uint32 ClassId;
+		FSoundModulationValue Value;
+		FBusHandle BusHandle;
+	};
+
+	class FModulatorBusMixProxy : public TModulatorProxyRefType<FBusMixId, FModulatorBusMixProxy, USoundControlBusMix>
+	{
+	public:
+		enum class EStatus : uint8
+		{
+			Enabled,
+			Stopping,
+			Stopped
+		};
+
+		FModulatorBusMixProxy(const USoundControlBusMix& InMix, FAudioModulationSystem& InModSystem);
+
+		FModulatorBusMixProxy& operator =(const USoundControlBusMix& InBusMix);
+
+		EStatus GetStatus() const;
+
+		// Resets channel map
+		void Reset();
+
+		void SetEnabled(const USoundControlBusMix& InBusMix);
+		void SetMix(const TArray<FSoundControlBusMixChannel>& InChannels);
+		void SetMixByFilter(const FString& InAddressFilter, uint32 InFilterClassId, const FSoundModulationValue& InValue);
+		void SetStopping();
+
+		void Update(const float Elapsed, FBusProxyMap& ProxyMap);
+
+		using FChannelMap = TMap<FBusId, FModulatorBusMixChannelProxy>;
+		FChannelMap Channels;
+
+	private:
+		EStatus Status;
+	};
+
+	using FBusMixProxyMap = TMap<FBusMixId, FModulatorBusMixProxy>;
+	using FBusMixHandle = TProxyHandle<FBusMixId, FModulatorBusMixProxy, USoundControlBusMix>;
+} // namespace AudioModulation
\ No newline at end of file
diff --git a/Engine/Plugins/Runtime/AudioModulation/Source/AudioModulation/Private/SoundControlBusProxy.cpp b/Engine/Plugins/Runtime/AudioModulation/Source/AudioModulation/Private/SoundControlBusProxy.cpp
new file mode 100644
--- /dev/null
+++ b/Engine/Plugins/Runtime/AudioModulation/Source/AudioModulation/Private/SoundControlBusProxy.cpp
@@ -0,0 +1,146 @@
+// Copyright Epic Games, Inc. All Rights Reserved.
+#include "SoundControlBusProxy.h"
+
+#include "AudioDevice.h"
+#include "AudioDeviceManager.h"
+#include "AudioModulation.h"
+#include "AudioModulationLogging.h"
+#include "AudioModulationSystem.h"
+#include "Engine/World.h"
+#include "SoundModulatorLFO.h"
+
+
+namespace AudioModulation
+{
+	const FBusId InvalidBusId = INDEX_NONE;
+
+	FControlBusProxy::FControlBusProxy()
+		: DefaultValue(0.0f)
+		, LFOValue(1.0f)
+		, MixValue(NAN)
+		, bBypass(false)
+		, Operator(ESoundModulatorOperator::Multiply)
+		, Range(0.0f, 1.0f)
+	{
+	}
+
+	FControlBusProxy::FControlBusProxy(const USoundControlBusBase& InBus, FAudioModulationSystem& InModSystem)
+		: TModulatorProxyRefType(InBus.GetName(), InBus.GetUniqueID(), InModSystem)
+	{
+		Init(InBus);
+	}
+
+	FControlBusProxy& FControlBusProxy::operator =(const USoundControlBusBase& InBus)
+	{
+		Init(InBus);
+		return *this;
+	}
+
+	float FControlBusProxy::GetDefaultValue() const
+	{
+		return DefaultValue;
+	}
+
+	const TArray<FLFOHandle>& FControlBusProxy::GetLFOHandles() const
+	{
+		return LFOHandles;
+	}
+
+	float FControlBusProxy::GetLFOValue() const
+	{
+		return LFOValue;
+	}
+
+	FVector2D FControlBusProxy::GetRange() const
+	{
+		return Range;
+	}
+
+	float FControlBusProxy::GetMixValue() const
+	{
+		return MixValue;
+	}
+
+	float FControlBusProxy::GetValue() const
+	{
+		const float DefaultMixed = Mix(DefaultValue);
+		return FMath::Clamp(DefaultMixed * LFOValue, Range.X, Range.Y);
+	}
+
+	void FControlBusProxy::Init(const USoundControlBusBase& InBus)
+	{
+		DefaultValue = InBus.DefaultValue;
+		LFOValue = 1.0f;
+		MixValue = NAN;
+		Operator = InBus.GetOperator();
+		Range = FVector2D(InBus.Min, InBus.Max);
+		if (InBus.Min > InBus.Max)
+		{
+			Range.X = InBus.Max;
+			Range.Y = InBus.Min;
+		}
+
+		DefaultValue = FMath::Clamp(DefaultValue, Range.X, Range.Y);
+		bBypass = InBus.bBypass;
+	}
+
+	void FControlBusProxy::InitLFOs(const USoundControlBusBase& InBus)
+	{
+		TArray<FLFOHandle> NewHandles;
+
+		for (const USoundBusModulatorBase* Modulator : InBus.Modulators)
+		{
+			if (const USoundBusModulatorLFO* LFO = Cast<USoundBusModulatorLFO>(Modulator))
+			{
+				check(ModSystem);
+				NewHandles.Add(FLFOHandle::Create(*LFO, ModSystem->RefProxies.LFOs, *ModSystem));
+			}
+		}
+
+		// Move vs. reset and adding to original array to avoid potentially clearing handles (and thus current LFO state)
+		// and destroying LFOs if function is called while reinitializing/updating the modulator
+		LFOHandles = MoveTemp(NewHandles);
+	}
+
+	bool FControlBusProxy::IsBypassed() const
+	{
+		return bBypass;
+	}
+
+	float FControlBusProxy::Mix(float ValueA) const
+	{
+		// If mix value is NaN, it is uninitialized (effectively, the parent bus is inactive)
+		// and therefore not mixable, so just return the second value.
+		if (FMath::IsNaN(MixValue))
+		{
+			return ValueA;
+		}
+		return SoundModulatorOperator::Apply(Operator, MixValue, ValueA);
+	}
+
+	void FControlBusProxy::MixIn(const float InValue)
+	{
+		MixValue = Mix(InValue);
+	}
+
+	void FControlBusProxy::MixLFO()
+	{
+		for (const FLFOHandle& Handle: LFOHandles)
+		{
+			if (Handle.IsValid())
+			{
+				const FModulatorLFOProxy& LFOProxy = Handle.FindProxy();
+				if (!LFOProxy.IsBypassed())
+				{
+					LFOValue *= LFOProxy.GetValue();
+				}
+			}
+		}
+	}
+
+	void FControlBusProxy::Reset()
+	{
+		LFOValue = 1.0f;
+		MixValue = NAN;
+	}
+} // namespace AudioModulation
diff --git a/Engine/Plugins/Runtime/AudioModulation/Source/AudioModulation/Private/SoundControlBusProxy.h b/Engine/Plugins/Runtime/AudioModulation/Source/AudioModulation/Private/SoundControlBusProxy.h
new file mode 100644
--- /dev/null
+++ b/Engine/Plugins/Runtime/AudioModulation/Source/AudioModulation/Private/SoundControlBusProxy.h
@@ -0,0 +1,57 @@
+// Copyright Epic Games, Inc. All Rights Reserved.
+#pragma once
+
+#include "IAudioModulation.h"
+#include "SoundControlBus.h"
+#include "SoundModulationProxy.h"
+#include "SoundModulatorLFOProxy.h"
+
+
+namespace AudioModulation
+{
+	// Forward Declarations
+	class FAudioModulationSystem;
+
+	using FBusId = uint32;
+	extern const FBusId InvalidBusId;
+
+	class FControlBusProxy : public TModulatorProxyRefType<FBusId, FControlBusProxy, USoundControlBusBase>
+	{
+	public:
+		FControlBusProxy();
+		FControlBusProxy(const USoundControlBusBase& Bus, FAudioModulationSystem& InModSystem);
+
+		FControlBusProxy& operator =(const USoundControlBusBase& InLFO);
+
+		float GetDefaultValue() const;
+		const TArray<FLFOHandle>& GetLFOHandles() const;
+		float GetLFOValue() const;
+		float GetMixValue() const;
+		FVector2D GetRange() const;
+		float GetValue() const;
+		void InitLFOs(const USoundControlBusBase& InBus);
+		bool IsBypassed() const;
+		void MixIn(const float InValue);
+		void MixLFO();
+		void Reset();
+
+	private:
+		void Init(const USoundControlBusBase& InBus);
+		float Mix(float ValueA) const;
+
+		float DefaultValue;
+
+		// Cached values
+		float LFOValue;
+		float MixValue;
+
+		bool bBypass;
+
+		TArray<FLFOHandle> LFOHandles;
+		ESoundModulatorOperator Operator;
+		FVector2D Range;
+	};
+
+	using FBusProxyMap = TMap<FBusId, FControlBusProxy>;
+	using FBusHandle = TProxyHandle<FBusId, FControlBusProxy, USoundControlBusBase>;
+} // namespace AudioModulation
\ No newline at end of file
diff --git a/Engine/Plugins/Runtime/AudioModulation/Source/AudioModulation/Private/SoundModulationPatch.cpp b/Engine/Plugins/Runtime/AudioModulation/Source/AudioModulation/Private/SoundModulationPatch.cpp
--- a/Engine/Plugins/Runtime/AudioModulation/Source/AudioModulation/Private/SoundModulationPatch.cpp
+++ b/Engine/Plugins/Runtime/AudioModulation/Source/AudioModulation/Private/SoundModulationPatch.cpp
@@ -4,13 +4,13 @@
 #include "AudioDefines.h"
 #include "AudioDevice.h"
 #include "AudioModulation.h"
-#include "AudioModulationInternal.h"
 #include "AudioModulationStatics.h"
-#include "SoundModulationProxy.h"
+#include "AudioModulationSystem.h"
+#include "SoundControlBus.h"
 #include "SoundModulationTransform.h"
 
 
-namespace
+namespace AudioModulation
 {
 	template <typename T>
 	void ClampPatchInputs(TArray<T>& Inputs)
@@ -23,61 +23,22 @@ namespace
 			}
 		}
 	}
-} // namespace <>
+} // namespace AudioModulation
 
-USoundModulationSettings::USoundModulationSettings(const FObjectInitializer& ObjectInitializer)
+USoundModulationPatch::USoundModulationPatch(const FObjectInitializer& ObjectInitializer)
 	: Super(ObjectInitializer)
 {
-	Volume.Output.Transform.Curve = ESoundModulatorOutputCurve::Exp;
-	Volume.Output.Transform.Scalar = 0.5f;
-
-	Lowpass.Output.Transform.OutputMin = 20.0f;
-	Lowpass.Output.Transform.OutputMax = 20000.0f;
-	Lowpass.Output.Transform.Curve = ESoundModulatorOutputCurve::Exp;
-	Lowpass.Output.SetOperator(ESoundModulatorOperator::Min);
-
-	Highpass.DefaultInputValue = 0.0f;
-	Highpass.Output.Transform.OutputMin = 20.0f;
-	Highpass.Output.Transform.OutputMax = 20000.0f;
-	Highpass.Output.Transform.Curve = ESoundModulatorOutputCurve::Exp;
-	Highpass.Output.SetOperator(ESoundModulatorOperator::Max);
 }
 
 #if WITH_EDITOR
-void USoundModulationSettings::OnPostEditChange(UWorld* World)
-{
-	Volume.Clamp();
-	Pitch.Clamp();
-	Lowpass.Clamp();
-	Highpass.Clamp();
-
-	for (FSoundControlModulationPatch& Control : Controls)
-	{
-		Control.Clamp();
-	}
-
-	if (AudioModulation::FAudioModulationImpl* Impl = UAudioModulationStatics::GetModulationImpl(World))
-	{
-		Impl->OnEditPluginSettings(*this);
-	}
-}
-
-void USoundModulationSettings::PostDuplicate(EDuplicateMode::Type DuplicateMode)
+void USoundModulationPatch::PostEditChangeProperty(FPropertyChangedEvent& InPropertyChangedEvent)
 {
-	if (DuplicateMode == EDuplicateMode::Normal)
+	if (FProperty* Property = InPropertyChangedEvent.Property)
 	{
-		OnPostEditChange(GetWorld());
+		AudioModulation::OnEditModulator(InPropertyChangedEvent, *this);
 	}
-}
-
-void USoundModulationSettings::PostEditChangeProperty(FPropertyChangedEvent& PropertyChangedEvent)
-{
-	OnPostEditChange(GetWorld());
-}
 
-void USoundModulationSettings::PostEditChangeChainProperty(FPropertyChangedChainEvent& PropertyChangedEvent)
-{
-	OnPostEditChange(GetWorld());
+	Super::PostEditChangeProperty(InPropertyChangedEvent);
 }
 
 void FSoundModulationPatchBase::Clamp()
@@ -100,7 +61,7 @@ void FSoundVolumeModulationPatch::Clamp()
 {
 	FSoundModulationPatchBase::Clamp();
 
-	ClampPatchInputs<FSoundVolumeModulationInput>(Inputs);
+	AudioModulation::ClampPatchInputs<FSoundVolumeModulationInput>(Inputs);
 	Output.Transform.OutputMin = FMath::Clamp(Output.Transform.OutputMin, 0.0f, MAX_VOLUME);
 	Output.Transform.OutputMax = FMath::Clamp(Output.Transform.OutputMax, 0.0f, MAX_VOLUME);
 }
@@ -109,7 +70,7 @@ void FSoundPitchModulationPatch::Clamp()
 {
 	FSoundModulationPatchBase::Clamp();
 
-	ClampPatchInputs<FSoundPitchModulationInput>(Inputs);
+	AudioModulation::ClampPatchInputs<FSoundPitchModulationInput>(Inputs);
 	Output.Transform.OutputMin = FMath::Clamp(Output.Transform.OutputMin, MIN_PITCH, MAX_PITCH);
 	Output.Transform.OutputMax = FMath::Clamp(Output.Transform.OutputMax, MIN_PITCH, MAX_PITCH);
 }
@@ -118,7 +79,7 @@ void FSoundLPFModulationPatch::Clamp()
 {
 	FSoundModulationPatchBase::Clamp();
 
-	ClampPatchInputs<FSoundLPFModulationInput>(Inputs);
+	AudioModulation::ClampPatchInputs<FSoundLPFModulationInput>(Inputs);
 	Output.Transform.OutputMin = FMath::Clamp(Output.Transform.OutputMin, MIN_FILTER_FREQUENCY, MAX_FILTER_FREQUENCY);
 	Output.Transform.OutputMax = FMath::Clamp(Output.Transform.OutputMax, MIN_FILTER_FREQUENCY, MAX_FILTER_FREQUENCY);
 }
@@ -127,7 +88,7 @@ void FSoundHPFModulationPatch::Clamp()
 {
 	FSoundModulationPatchBase::Clamp();
 
-	ClampPatchInputs<FSoundHPFModulationInput>(Inputs);
+	AudioModulation::ClampPatchInputs<FSoundHPFModulationInput>(Inputs);
 	Output.Transform.OutputMin = FMath::Clamp(Output.Transform.OutputMin, MIN_FILTER_FREQUENCY, MAX_FILTER_FREQUENCY);
 	Output.Transform.OutputMax = FMath::Clamp(Output.Transform.OutputMax, MIN_FILTER_FREQUENCY, MAX_FILTER_FREQUENCY);
 }
@@ -233,89 +194,3 @@ FSoundModulationPatchBase::FSoundModulationPatchBase()
 	, bBypass(1)
 {
 }
-
-namespace AudioModulation
-{
-	FModulationInputProxy::FModulationInputProxy()
-		: bSampleAndHold(0)
-	{
-	}
-
-	FModulationInputProxy::FModulationInputProxy(const FSoundModulationInputBase& Input, FReferencedProxies& OutRefProxies, FAudioModulationImpl& ModulationImpl)
-		: Transform(Input.Transform)
-		, bSampleAndHold(Input.bSampleAndHold)
-	{
-		if (const USoundControlBusBase* Bus = Input.GetBus())
-		{
-			FReferencedProxies* RefProxies = &OutRefProxies;
-			auto OnCreate = [RefProxies, Bus](FControlBusProxy& NewProxy)
-			{
-				NewProxy.InitLFOs(*Bus, RefProxies->LFOs);
-			};
-			BusHandle = FBusHandle::Create(*Bus, OutRefProxies.Buses, ModulationImpl, OnCreate);
-		}
-	}
-
-	FModulationOutputProxy::FModulationOutputProxy()
-		: bInitialized(0)
-		, Operator(ESoundModulatorOperator::Multiply)
-		, SampleAndHoldValue(1.0f)
-	{
-	}
-
-	FModulationOutputProxy::FModulationOutputProxy(const FSoundModulationOutputBase& Output)
-		: bInitialized(0)
-		, Operator(Output.GetOperator())
-		, SampleAndHoldValue(1.0f)
-		, Transform(Output.Transform)
-	{
-	}
-
-	FModulationPatchProxy::FModulationPatchProxy()
-		: DefaultInputValue(1.0f)
-		, bBypass(1)
-	{
-	}
-
-	FModulationPatchProxy::FModulationPatchProxy(const FSoundModulationPatchBase& Patch, FReferencedProxies& OutRefProxies, FAudioModulationImpl& InModulationImpl)
-		: DefaultInputValue(Patch.DefaultInputValue)
-		, bBypass(Patch.bBypass)
-		, OutputProxy(*Patch.GetOutput())
-	{
-		TArray<const FSoundModulationInputBase*> Inputs = Patch.GetInputs();
-		for (const FSoundModulationInputBase* Input : Inputs)
-		{
-			InputProxies.Emplace_GetRef(*Input, OutRefProxies, InModulationImpl);
-		}
-	}
-
-	FModulationSettingsProxy::FModulationSettingsProxy()
-	{
-		Volume.OutputProxy.Transform.Curve = ESoundModulatorOutputCurve::Exp;
-		Volume.OutputProxy.Transform.Scalar = 0.5f;
-
-		Lowpass.OutputProxy.Transform.OutputMin = 20.0f;
-		Lowpass.OutputProxy.Transform.OutputMax = 20000.0f;
-		Lowpass.OutputProxy.Transform.Curve = ESoundModulatorOutputCurve::Exp;
-		Lowpass.OutputProxy.Operator = ESoundModulatorOperator::Min;
-
-		Highpass.DefaultInputValue = 0.0f;
-		Highpass.OutputProxy.Transform.OutputMin = 20.0f;
-		Highpass.OutputProxy.Transform.OutputMax = 20000.0f;
-		Highpass.OutputProxy.Transform.Curve = ESoundModulatorOutputCurve::Exp;
-		Highpass.OutputProxy.Operator = ESoundModulatorOperator::Max;
-	}
-
-	FModulationSettingsProxy::FModulationSettingsProxy(const USoundModulationSettings& Settings, FReferencedProxies& OutRefProxies, FAudioModulationImpl& InModulationImpl)
-		: TModulatorProxyBase<uint32>(Settings.GetName(), Settings.GetUniqueID())
-		, Volume(Settings.Volume, OutRefProxies, InModulationImpl)
-		, Pitch(Settings.Pitch, OutRefProxies, InModulationImpl)
-		, Lowpass(Settings.Lowpass, OutRefProxies, InModulationImpl)
-		, Highpass(Settings.Highpass, OutRefProxies, InModulationImpl)
-	{
-		for (const FSoundControlModulationPatch& Patch : Settings.Controls)
-		{
-			Controls.Add(Patch.Control, FModulationPatchProxy(Patch, OutRefProxies, InModulationImpl));
-		}
-	}
-} // namespace AudioModulation
\ No newline at end of file
diff --git a/Engine/Plugins/Runtime/AudioModulation/Source/AudioModulation/Private/SoundModulationPatchProxy.cpp b/Engine/Plugins/Runtime/AudioModulation/Source/AudioModulation/Private/SoundModulationPatchProxy.cpp
new file mode 100644
--- /dev/null
+++ b/Engine/Plugins/Runtime/AudioModulation/Source/AudioModulation/Private/SoundModulationPatchProxy.cpp
@@ -0,0 +1,184 @@
+// Copyright Epic Games, Inc. All Rights Reserved.
+#include "SoundModulationPatchProxy.h"
+
+#include "AudioDefines.h"
+#include "AudioDevice.h"
+#include "AudioModulation.h"
+#include "AudioModulationStatics.h"
+#include "AudioModulationSystem.h"
+#include "SoundControlBusProxy.h"
+#include "SoundModulationProxy.h"
+#include "SoundModulationTransform.h"
+
+
+namespace AudioModulation
+{
+	const FPatchId InvalidPatchId = INDEX_NONE;
+
+	void MixInModulationValue(ESoundModulatorOperator& Operator, float ModStageValue, float& Value)
+	{
+		switch (Operator)
+		{
+			case ESoundModulatorOperator::Max:
+			{
+				Value = FMath::Max(ModStageValue, Value);
+			}
+			break;
+
+			case ESoundModulatorOperator::Min:
+			{
+				Value = FMath::Min(ModStageValue, Value);
+			}
+			break;
+
+			case ESoundModulatorOperator::Multiply:
+			{
+				Value *= ModStageValue;
+			}
+			break;
+
+			case ESoundModulatorOperator::Divide:
+			{
+				Value /= ModStageValue;
+			}
+			break;
+
+			case ESoundModulatorOperator::Add:
+			{
+				Value += ModStageValue;
+			}
+			break;
+
+			case ESoundModulatorOperator::Subtract:
+			{
+				Value -= ModStageValue;
+			}
+			break;
+
+			case ESoundModulatorOperator::None:
+			default:
+			{
+				checkf(false, TEXT("Cannot apply 'None' as operator to modulator"));
+				static_assert(static_cast<int32>(ESoundModulatorOperator::Count) == 7, "Possible missing operator switch case coverage");
+			}
+			break;
+		}
+	}
+
+	FModulationInputProxy::FModulationInputProxy(const FSoundModulationInputBase& Input, FAudioModulationSystem& ModSystem)
+		: Transform(Input.Transform)
+		, bSampleAndHold(Input.bSampleAndHold)
+	{
+		if (const USoundControlBusBase* Bus = Input.GetBus())
+		{
+			auto OnCreate = [Bus](FControlBusProxy& NewProxy)
+			{
+				NewProxy.InitLFOs(*Bus);
+			};
+			BusHandle = FBusHandle::Create(*Bus, ModSystem.RefProxies.Buses, ModSystem, OnCreate);
+		}
+	}
+
+	FModulationOutputProxy::FModulationOutputProxy(const FSoundModulationOutputBase& Output)
+		: Operator(Output.GetOperator())
+		, Transform(Output.Transform)
+	{
+	}
+
+	FModulationPatchProxy::FModulationPatchProxy(const FSoundModulationPatchBase& InPatch, FAudioModulationSystem& InModSystem)
+	{
+		Init(InPatch, InModSystem);
+	}
+
+	void FModulationPatchProxy::Init(const FSoundModulationPatchBase& InPatch, FAudioModulationSystem& InModSystem)
+	{
+		bBypass = InPatch.bBypass;
+
+		DefaultInputValue = InPatch.DefaultInputValue;
+
+		InputProxies.Reset();
+		TArray<const FSoundModulationInputBase*> Inputs = InPatch.GetInputs();
+		for (const FSoundModulationInputBase* Input : Inputs)
+		{
+			InputProxies.Emplace_GetRef(*Input, InModSystem);
+		}
+
+		OutputProxy = *InPatch.GetOutput();
+	}
+
+	bool FModulationPatchProxy::IsBypassed() const
+	{
+		return bBypass;
+	}
+
+	float FModulationPatchProxy::Update()
+	{
+		float OutValue = DefaultInputValue;
+
+		float& OutSampleHold = OutputProxy.SampleAndHoldValue;
+		if (!OutputProxy.bInitialized)
+		{
+			OutSampleHold = SoundModulatorOperator::GetDefaultValue(OutputProxy.Operator, OutputProxy.Transform.OutputMin, OutputProxy.Transform.OutputMax);
+		}
+
+		for (const FModulationInputProxy& Input : InputProxies)
+		{
+			if (Input.bSampleAndHold)
+			{
+				if (!OutputProxy.bInitialized && Input.BusHandle.IsValid())
+				{
+					const FControlBusProxy& BusProxy = Input.BusHandle.FindProxy();
+					if (!BusProxy.IsBypassed())
+					{
+						float ModStageValue = BusProxy.GetValue();
+						Input.Transform.Apply(ModStageValue);
+						MixInModulationValue(OutputProxy.Operator, ModStageValue, OutSampleHold);
+					}
+				}
+			}
+			else
+			{
+				if (Input.BusHandle.IsValid())
+				{
+					const FControlBusProxy& BusProxy = Input.BusHandle.FindProxy();
+					if (!BusProxy.IsBypassed())
+					{
+						float ModStageValue = BusProxy.GetValue();
+						Input.Transform.Apply(ModStageValue);
+						MixInModulationValue(OutputProxy.Operator, ModStageValue, OutValue);
+					}
+				}
+			}
+		}
+
+		if (!OutputProxy.bInitialized)
+		{
+			const float OutputMin = OutputProxy.Transform.OutputMin;
+			const float OutputMax = OutputProxy.Transform.OutputMax;
+			OutSampleHold = FMath::Clamp(OutSampleHold, OutputMin, OutputMax);
+			OutputProxy.bInitialized = true;
+		}
+
+		OutputProxy.Transform.Apply(OutValue);
+		MixInModulationValue(OutputProxy.Operator, OutSampleHold, OutValue);
+		return OutValue;
+	}
+
+	FModulationPatchRefProxy::FModulationPatchRefProxy()
+		: FModulationPatchProxy()
+		, TModulatorProxyRefType()
+	{
+	}
+
+	FModulationPatchRefProxy::FModulationPatchRefProxy(const USoundModulationPatch& InPatch, FAudioModulationSystem& InModSystem)
+		: FModulationPatchProxy(InPatch.PatchSettings, InModSystem)
+		, TModulatorProxyRefType(InPatch.GetName(), InPatch.GetUniqueID(), InModSystem)
+	{
+	}
+
+	FModulationPatchRefProxy& FModulationPatchRefProxy::operator =(const USoundModulationPatch& InPatch)
+	{
+		Init(InPatch.PatchSettings, *ModSystem);
+		return *this;
+	}
+} // namespace AudioModulation
\ No newline at end of file
diff --git a/Engine/Plugins/Runtime/AudioModulation/Source/AudioModulation/Private/SoundModulationPatchProxy.h b/Engine/Plugins/Runtime/AudioModulation/Source/AudioModulation/Private/SoundModulationPatchProxy.h
new file mode 100644
--- /dev/null
+++ b/Engine/Plugins/Runtime/AudioModulation/Source/AudioModulation/Private/SoundModulationPatchProxy.h
@@ -0,0 +1,92 @@
+// Copyright Epic Games, Inc. All Rights Reserved.
+#pragma once
+
+#include "SoundControlBusProxy.h"
+#include "SoundControlBusMixProxy.h"
+#include "SoundModulationProxy.h"
+#include "SoundModulatorLFOProxy.h"
+
+
+namespace AudioModulation
+{
+	// Forward Declarations
+	class FAudioModulationSystem;
+
+	using FPatchId = uint32;
+	extern const FPatchId InvalidPatchId;
+
+	/** Modulation input instance */
+	class FModulationInputProxy
+	{
+	public:
+
+		FModulationInputProxy() = default;
+		FModulationInputProxy(const FSoundModulationInputBase& Patch, FAudioModulationSystem& ModSystem);
+
+		FBusHandle BusHandle;
+		FSoundModulationInputTransform Transform;
+		bool bSampleAndHold = false;
+	};
+
+	/** Patch applied as the final stage of a modulation chain prior to output on the sound level (Always active, never removed) */
+	struct FModulationOutputProxy
+	{
+		FModulationOutputProxy() = default;
+		FModulationOutputProxy(const FSoundModulationOutputBase& Patch);
+
+		/** Whether patch has been initialized or not */
+		bool bInitialized = false;
+
+		/** Operator used to calculate the output proxy value */
+		ESoundModulatorOperator Operator = ESoundModulatorOperator::Multiply;
+
+		/** Cached value of sample-and-hold input values */
+		float SampleAndHoldValue = 1.0f;
+
+		/** Final transform before passing to output */
+		FSoundModulationOutputTransform Transform;
+	};
+
+	class FModulationPatchProxy
+	{
+	public:
+		FModulationPatchProxy() = default;
+		FModulationPatchProxy(const FSoundModulationPatchBase& InPatch, FAudioModulationSystem& InModSystem);
+
+		bool IsBypassed() const;
+
+		// Updates the patch and returns the current value
+		float Update();
+
+	protected:
+		void Init(const FSoundModulationPatchBase& InPatch, FAudioModulationSystem& InModSystem);
+
+	private:
+		/** Default value of patch (Value mixed when inputs are provided or not, regardless of active state)*/
+		float DefaultInputValue = 1.0f;
+
+		/** Optional modulation inputs */
+		TArray<FModulationInputProxy> InputProxies;
+
+		/** Final output modulation post input combination */
+		FModulationOutputProxy OutputProxy;
+
+		/** Bypasses the patch and doesn't update modulation value */
+		bool bBypass = true;
+
+		friend class FAudioModulationSystem;
+		friend class FModulationSettingsProxy;
+	};
+
+	class FModulationPatchRefProxy : public FModulationPatchProxy, public TModulatorProxyRefType<FPatchId, FModulationPatchRefProxy, USoundModulationPatch>
+	{
+	public:
+		FModulationPatchRefProxy();
+		FModulationPatchRefProxy(const USoundModulationPatch& InPatch, FAudioModulationSystem& InModSystem);
+
+		FModulationPatchRefProxy& operator =(const USoundModulationPatch& InPatch);
+	};
+
+	using FPatchProxyMap = TMap<FPatchId, FModulationPatchRefProxy>;
+	using FPatchHandle = TProxyHandle<FPatchId, FModulationPatchRefProxy, USoundModulationPatch>;
+} // namespace AudioModulation
\ No newline at end of file
diff --git a/Engine/Plugins/Runtime/AudioModulation/Source/AudioModulation/Private/SoundModulationProxy.h b/Engine/Plugins/Runtime/AudioModulation/Source/AudioModulation/Private/SoundModulationProxy.h
--- a/Engine/Plugins/Runtime/AudioModulation/Source/AudioModulation/Private/SoundModulationProxy.h
+++ b/Engine/Plugins/Runtime/AudioModulation/Source/AudioModulation/Private/SoundModulationProxy.h
@@ -1,40 +1,20 @@
 // Copyright Epic Games, Inc. All Rights Reserved.
 #pragma once
 
-#include "CoreMinimal.h"
-#include "IAudioExtensionPlugin.h"
-#include "UObject/ObjectMacros.h"
-#include "UObject/Object.h"
-#include "UObject/WeakObjectPtrTemplates.h"
-
 #include "AudioModulationLogging.h"
-#include "SoundControlBus.h"
-#include "SoundControlBusMix.h"
+#include "CoreMinimal.h"
 #include "SoundModulationPatch.h"
-#include "SoundModulationTransform.h"
-#include "SoundModulatorLFO.h"
 
 
 namespace AudioModulation
 {
 	// Forward Declarations
-	class FAudioModulationImpl;
-	struct FReferencedProxies;
-
-	// Modulator Ids
-	using FBusMixId = uint32;
-	extern const FBusMixId InvalidBusMixId;
-
-	using FBusId = uint32;
-	extern const FBusId InvalidBusId;
-
-	using FLFOId = uint32;
-	extern const FBusId InvalidLFOId;
+	class FAudioModulationSystem;
 
 	/*
 	 * Handle for all ref-counted proxy types, to be used only on the audio thread (not thread safe). 
 	 */
-	template<typename IdType, typename ProxyType, typename ProxyUObjType>
+	template<typename IdType, typename ProxyType, typename ProxiedObject>
 	class TProxyHandle
 	{
 		IdType Id;
@@ -67,7 +47,7 @@ namespace AudioModulation
 			return ProxyMap->FindChecked(Id);
 		}
 
-		friend FORCEINLINE uint32 GetTypeHash(const TProxyHandle<IdType, ProxyType, ProxyUObjType>& InHandle)
+		friend FORCEINLINE uint32 GetTypeHash(const TProxyHandle<IdType, ProxyType, ProxiedObject>& InHandle)
 		{
 			return static_cast<uint32>(InHandle.Id);
 		}
@@ -132,22 +112,21 @@ namespace AudioModulation
 		/*
 		 * Creates a handle to a proxy modulation object tracked in the provided InProxyMap if it exists, otherwise returns invalid handle.
 		 */
-		static TProxyHandle<IdType, ProxyType, ProxyUObjType> Get(const IdType ObjectId, TMap<IdType, ProxyType>& InProxyMap)
+		static TProxyHandle<IdType, ProxyType, ProxiedObject> Get(const IdType ObjectId, TMap<IdType, ProxyType>& InProxyMap)
 		{
 			if (ProxyType* Proxy = InProxyMap.Find(ObjectId))
 			{
-				check(Proxy->ModulationImpl);
-				TProxyHandle<IdType, ProxyType, ProxyUObjType> NewHandle(ObjectId, InProxyMap);
-				return MoveTemp(NewHandle);
+				check(Proxy->ModSystem);
+				return TProxyHandle<IdType, ProxyType, ProxiedObject>(ObjectId, InProxyMap);
 			} 
 
-			return TProxyHandle<IdType, ProxyType, ProxyUObjType>();
+			return TProxyHandle<IdType, ProxyType, ProxiedObject>();
 		}
 
 		/*
 		 * Creates a handle to a proxy modulation object tracked in the provided InProxyMap if it exists, otherwise returns invalid handle.
 		 */
-		static TProxyHandle<IdType, ProxyType, ProxyUObjType> Get(const ProxyUObjType& InObject, TMap<IdType, ProxyType>& InProxyMap)
+		static TProxyHandle<IdType, ProxyType, ProxiedObject> Get(const ProxiedObject& InObject, TMap<IdType, ProxyType>& InProxyMap)
 		{
 			const IdType ObjectId = static_cast<IdType>(InObject.GetUniqueID());
 			return Get(ObjectId, InProxyMap);
@@ -156,42 +135,42 @@ namespace AudioModulation
 		/*
 		 * Creates a handle to a proxy modulation object tracked in the provided InProxyMap.  Creates new proxy if it doesn't exist.
 		 */
-		static TProxyHandle<IdType, ProxyType, ProxyUObjType> Create(const ProxyUObjType& InObject, TMap<IdType, ProxyType>& InProxyMap, FAudioModulationImpl& InModulationImpl)
+		static TProxyHandle<IdType, ProxyType, ProxiedObject> Create(const ProxiedObject& InObject, TMap<IdType, ProxyType>& InProxyMap, FAudioModulationSystem& InModSystem)
 		{
 			const IdType ObjectId = static_cast<IdType>(InObject.GetUniqueID());
 
-			TProxyHandle<IdType, ProxyType, ProxyUObjType> NewHandle = Get(InObject, InProxyMap);
+			TProxyHandle<IdType, ProxyType, ProxiedObject> NewHandle = Get(InObject, InProxyMap);
 			if (!NewHandle.IsValid())
 			{
 				UE_LOG(LogAudioModulation, Verbose, TEXT("Proxy created: Id '%u' for object '%s'."), NewHandle.Id, *InObject.GetName());
-				InProxyMap.Add(ObjectId, ProxyType(InObject, InModulationImpl));
-				NewHandle = TProxyHandle<IdType, ProxyType, ProxyUObjType>(ObjectId, InProxyMap);
+				InProxyMap.Add(ObjectId, ProxyType(InObject, InModSystem));
+				NewHandle = TProxyHandle<IdType, ProxyType, ProxiedObject>(ObjectId, InProxyMap);
 			}
 
 			return NewHandle;
 		}
 
-		static TProxyHandle<IdType, ProxyType, ProxyUObjType> Create(const ProxyUObjType& InObject, TMap<IdType, ProxyType>& InProxyMap, FAudioModulationImpl& InModulationImpl, TFunction<void(ProxyType&)> OnCreateProxy)
+		static TProxyHandle<IdType, ProxyType, ProxiedObject> Create(const ProxiedObject& InObject, TMap<IdType, ProxyType>& InProxyMap, FAudioModulationSystem& InModSystem, TFunction<void(ProxyType&)> OnCreateProxy)
 		{
 			const IdType ObjectId = static_cast<IdType>(InObject.GetUniqueID());
-			TProxyHandle<IdType, ProxyType, ProxyUObjType> NewHandle = Get(InObject, InProxyMap);
+			TProxyHandle<IdType, ProxyType, ProxiedObject> NewHandle = Get(InObject, InProxyMap);
 			if (!NewHandle.IsValid())
 			{
 				UE_LOG(LogAudioModulation, Verbose, TEXT("Proxy created: Id '%u' for object '%s'."), NewHandle.Id, *InObject.GetName());
-				InProxyMap.Add(ObjectId, ProxyType(InObject, InModulationImpl));
-				NewHandle = TProxyHandle<IdType, ProxyType, ProxyUObjType>(ObjectId, InProxyMap);
+				InProxyMap.Add(ObjectId, ProxyType(InObject, InModSystem));
+				NewHandle = TProxyHandle<IdType, ProxyType, ProxiedObject>(ObjectId, InProxyMap);
 				OnCreateProxy(NewHandle.FindProxy());
 			}
 
 			return NewHandle;
 		}
 
-		FORCEINLINE bool operator ==(const TProxyHandle<IdType, ProxyType, ProxyUObjType>& InHandle) const
+		FORCEINLINE bool operator ==(const TProxyHandle<IdType, ProxyType, ProxiedObject>& InHandle) const
 		{
 			return InHandle.Id == Id;
 		}
 
-		FORCEINLINE TProxyHandle<IdType, ProxyType, ProxyUObjType>& operator =(const TProxyHandle<IdType, ProxyType, ProxyUObjType>& InHandle)
+		FORCEINLINE TProxyHandle<IdType, ProxyType, ProxiedObject>& operator =(const TProxyHandle<IdType, ProxyType, ProxiedObject>& InHandle)
 		{
 			// 1. If local proxy valid prior to move, cache to DecRef
 			ProxyType* ProxyToDecRef = nullptr;
@@ -220,7 +199,7 @@ namespace AudioModulation
 			return *this;
 		}
 
-		FORCEINLINE TProxyHandle<IdType, ProxyType, ProxyUObjType>& operator =(TProxyHandle<IdType, ProxyType, ProxyUObjType>&& InHandle)
+		FORCEINLINE TProxyHandle<IdType, ProxyType, ProxiedObject>& operator =(TProxyHandle<IdType, ProxyType, ProxiedObject>&& InHandle)
 		{
 			// 1. If local proxy valid prior to move, cache to DecRef
 			ProxyType* ProxyToDecRef = nullptr;
@@ -303,63 +282,58 @@ namespace AudioModulation
 		}
 	};
 
-	class FModulatorLFOProxy;
-
-	using FLFOProxyMap = TMap<FLFOId, FModulatorLFOProxy>;
-	using FLFOHandle = TProxyHandle<FLFOId, FModulatorLFOProxy, USoundBusModulatorLFO>;
-
-	template<typename IdType, typename ProxyType, typename ProxyUObjType>
+	template<typename IdType, typename ProxyType, typename ProxiedObject>
 	class TModulatorProxyRefType : public TModulatorProxyBase<IdType>
 	{
 	protected:
 		uint32 RefCount;
-		FAudioModulationImpl* ModulationImpl;
+		FAudioModulationSystem* ModSystem;
 
 	public:
 		TModulatorProxyRefType()
 			: TModulatorProxyBase<IdType>()
 			, RefCount(0)
-			, ModulationImpl(nullptr)
+			, ModSystem(nullptr)
 		{
 		}
 
-		TModulatorProxyRefType(const FString& InName, const uint32 InId, FAudioModulationImpl& InModulationImpl)
+		TModulatorProxyRefType(const FString& InName, const uint32 InId, FAudioModulationSystem& InModSystem)
 			: TModulatorProxyBase<IdType>(InName, InId)
 			, RefCount(0)
-			, ModulationImpl(&InModulationImpl)
+			, ModSystem(&InModSystem)
 		{
 		}
 
 		TModulatorProxyRefType(const TModulatorProxyRefType& InProxyRef)
 			: TModulatorProxyBase<IdType>(InProxyRef.GetName(), InProxyRef.GetId())
 			, RefCount(InProxyRef.RefCount)
-			, ModulationImpl(InProxyRef.ModulationImpl)
+			, ModSystem(InProxyRef.ModSystem)
 		{
 		}
 
 		TModulatorProxyRefType(TModulatorProxyRefType&& InProxyRef)
 			: TModulatorProxyBase<IdType>(InProxyRef.GetName(), InProxyRef.GetId())
 			, RefCount(InProxyRef.RefCount)
-			, ModulationImpl(InProxyRef.ModulationImpl)
+			, ModSystem(InProxyRef.ModSystem)
 		{
 			InProxyRef.RefCount = 0;
-			InProxyRef.ModulationImpl = nullptr;
+			InProxyRef.ModSystem = nullptr;
 		}
 
 		TModulatorProxyRefType& operator=(const TModulatorProxyRefType& InOther)
 		{
 			RefCount = InOther.RefCount;
-			ModulationImpl = InOther.ModulationImpl;
+			ModSystem = InOther.ModSystem;
 			return *this;
 		}
 
-		TModulatorProxyRefType& operator=(const TModulatorProxyRefType&& InOther)
+		TModulatorProxyRefType& operator=(TModulatorProxyRefType&& InOther)
 		{
 			RefCount = InOther.RefCount;
-			ModulationImpl = InOther.ModulationImpl;
+			ModSystem = InOther.ModSystem;
 
 			InOther.RefCount = 0;
-			InOther.ModulationImpl = nullptr;
+			InOther.ModSystem = nullptr;
 
 			return *this;
 		}
@@ -389,181 +363,6 @@ namespace AudioModulation
 		}
 
 	private:
-		friend class TProxyHandle<IdType, ProxyType, ProxyUObjType>;
-	};
-
-	class FModulatorLFOProxy : public TModulatorProxyRefType<FLFOId, FModulatorLFOProxy, USoundBusModulatorLFO>
-	{
-	public:
-		FModulatorLFOProxy();
-		FModulatorLFOProxy(const USoundBusModulatorLFO& InLFO, FAudioModulationImpl& InModulationImpl);
-
-		FModulatorLFOProxy& operator =(const USoundBusModulatorLFO& InLFO);
-
-		float GetValue() const;
-		bool IsBypassed() const;
-		void Update(float InElapsed);
-
-	private:
-		void Init(const USoundBusModulatorLFO& InLFO);
-
-		Audio::FLFO LFO;
-		float Offset;
-		float Value;
-		uint8 bBypass : 1;
-	};
-
-	class FControlBusProxy : public TModulatorProxyRefType<FBusId, FControlBusProxy, USoundControlBusBase>
-	{
-	public:
-		FControlBusProxy();
-		FControlBusProxy(const USoundControlBusBase& Bus, FAudioModulationImpl& InModulationImpl);
-
-		FControlBusProxy& operator =(const USoundControlBusBase& InLFO);
-
-		float GetDefaultValue() const;
-		const TArray<FLFOHandle>& GetLFOHandles() const;
-		float GetLFOValue() const;
-		float GetMixValue() const;
-		FVector2D GetRange() const;
-		float GetValue() const;
-		void InitLFOs(const USoundControlBusBase& InBus, FLFOProxyMap& OutActiveLFOs);
-		bool IsBypassed() const;
-		void MixIn(const float InValue);
-		void MixLFO();
-		void Reset();
-
-	private:
-		void Init(const USoundControlBusBase& InBus);
-		float Mix(float ValueA) const;
-		float Mix(float ValueA, float ValueB) const;
-
-		float DefaultValue;
-
-		// Cached values
-		float LFOValue;
-		float MixValue;
-
-		bool bBypass;
-
-		TArray<FLFOHandle> LFOHandles;
-		ESoundModulatorOperator Operator;
-		FVector2D Range;
-	};
-
-	using FBusProxyMap = TMap<FBusId, FControlBusProxy>;
-	using FBusHandle = TProxyHandle<FBusId, FControlBusProxy, USoundControlBusBase>;
-
-	struct FModulatorBusMixChannelProxy : public TModulatorProxyBase<FBusId>
-	{
-		FModulatorBusMixChannelProxy(const FSoundControlBusMixChannel& Channel, FAudioModulationImpl& ModulationImpl);
-		FString Address;
-		uint32 ClassId;
-		FSoundModulationValue Value;
-		FBusHandle BusHandle;
-	};
-
-	class FModulatorBusMixProxy : public TModulatorProxyRefType<FBusMixId, FModulatorBusMixProxy, USoundControlBusMix>
-	{
-	public:
-		enum class EStatus : uint8
-		{
-			Enabled,
-			Stopping,
-			Stopped
-		};
-
-		FModulatorBusMixProxy(const USoundControlBusMix& InMix, FAudioModulationImpl& InModulationImpl);
-
-		FModulatorBusMixProxy& operator =(const USoundControlBusMix& InBusMix);
-
-		EStatus GetStatus() const;
-
-		// Resets channel map
-		void Reset();
-
-		void SetEnabled(const USoundControlBusMix& InBusMix);
-		void SetMix(const TArray<FSoundControlBusMixChannel>& InChannels);
-		void SetMixByFilter(const FString& InAddressFilter, uint32 InFilterClassId, const FSoundModulationValue& InValue);
-		void SetStopping();
-
-		void Update(const float Elapsed, FBusProxyMap& ProxyMap);
-
-		using FChannelMap = TMap<FBusId, FModulatorBusMixChannelProxy>;
-		FChannelMap Channels;
-
-	private:
-		EStatus Status;
-	};
-
-	using FBusMixProxyMap = TMap<FBusMixId, FModulatorBusMixProxy>;
-	using FBusMixHandle = TProxyHandle<FBusMixId, FModulatorBusMixProxy, USoundControlBusMix>;
-
-	/** Modulation input instance */
-	struct FModulationInputProxy
-	{
-		FModulationInputProxy();
-		FModulationInputProxy(const FSoundModulationInputBase& Patch, FReferencedProxies& OutRefProxies, FAudioModulationImpl& ModulationImpl);
-
-		FBusHandle BusHandle;
-		FSoundModulationInputTransform Transform;
-		uint8 bSampleAndHold : 1;
-	};
-
-	/** Patch applied as the final stage of a modulation chain prior to output on the sound level (Always active, never removed) */
-	struct FModulationOutputProxy
-	{
-		FModulationOutputProxy();
-		FModulationOutputProxy(const FSoundModulationOutputBase& Patch);
-
-		/** Whether patch has been initialized or not */
-		uint8 bInitialized : 1;
-
-		/** Operator used to calculate the output proxy value */
-		ESoundModulatorOperator Operator;
-
-		/** Cached value of sample-and-hold input values */
-		float SampleAndHoldValue;
-
-		/** Final transform before passing to output */
-		FSoundModulationOutputTransform Transform;
-	};
-
-	struct FModulationPatchProxy
-	{
-		FModulationPatchProxy();
-		FModulationPatchProxy(const FSoundModulationPatchBase& Patch, FReferencedProxies& OutRefProxies, FAudioModulationImpl& InModulationImpl);
-
-		/** Default value of patch (Value mixed when inputs are provided or not, regardless of active state)*/
-		float DefaultInputValue;
-
-		/** Bypasses the patch and doesn't update modulation value */
-		uint8 bBypass : 1;
-
-		/** Optional modulation inputs */
-		TArray<FModulationInputProxy> InputProxies;
-
-		/** Final output modulation post input combination */
-		FModulationOutputProxy OutputProxy;
-	};
-
-	struct FModulationSettingsProxy : public TModulatorProxyBase<uint32>
-	{
-		FModulationSettingsProxy();
-		FModulationSettingsProxy(const USoundModulationSettings& Settings, FReferencedProxies& OutRefProxies, FAudioModulationImpl& InModulationImpl);
-
-		FModulationPatchProxy Volume;
-		FModulationPatchProxy Pitch;
-		FModulationPatchProxy Lowpass;
-		FModulationPatchProxy Highpass;
-
-		TMap<FName, FModulationPatchProxy> Controls;
-	};
-
-	struct FReferencedProxies
-	{
-		FBusMixProxyMap BusMixes;
-		FBusProxyMap    Buses;
-		FLFOProxyMap    LFOs;
+		friend class TProxyHandle<IdType, ProxyType, ProxiedObject>;
 	};
 } // namespace AudioModulation
\ No newline at end of file
diff --git a/Engine/Plugins/Runtime/AudioModulation/Source/AudioModulation/Private/SoundModulationSettings.cpp b/Engine/Plugins/Runtime/AudioModulation/Source/AudioModulation/Private/SoundModulationSettings.cpp
new file mode 100644
--- /dev/null
+++ b/Engine/Plugins/Runtime/AudioModulation/Source/AudioModulation/Private/SoundModulationSettings.cpp
@@ -0,0 +1,61 @@
+// Copyright Epic Games, Inc. All Rights Reserved.
+#include "SoundModulationSettings.h"
+
+#include "AudioDefines.h"
+#include "AudioDevice.h"
+#include "AudioModulation.h"
+#include "AudioModulationStatics.h"
+#include "AudioModulationSystem.h"
+#include "SoundModulationTransform.h"
+
+
+USoundModulationSettings::USoundModulationSettings(const FObjectInitializer& ObjectInitializer)
+	: Super(ObjectInitializer)
+{
+	Volume.Output.Transform.Curve = ESoundModulatorOutputCurve::Exp;
+	Volume.Output.Transform.Scalar = 0.5f;
+
+	Lowpass.Output.Transform.OutputMin = 20.0f;
+	Lowpass.Output.Transform.OutputMax = 20000.0f;
+	Lowpass.Output.Transform.Curve = ESoundModulatorOutputCurve::Exp;
+	Lowpass.Output.SetOperator(ESoundModulatorOperator::Min);
+
+	Highpass.DefaultInputValue = 0.0f;
+	Highpass.Output.Transform.OutputMin = 20.0f;
+	Highpass.Output.Transform.OutputMax = 20000.0f;
+	Highpass.Output.Transform.Curve = ESoundModulatorOutputCurve::Exp;
+	Highpass.Output.SetOperator(ESoundModulatorOperator::Max);
+}
+
+#if WITH_EDITOR
+void USoundModulationSettings::OnPostEditChange(UWorld* World)
+{
+	Volume.Clamp();
+	Pitch.Clamp();
+	Lowpass.Clamp();
+	Highpass.Clamp();
+
+	if (AudioModulation::FAudioModulationSystem* ModSystem = UAudioModulationStatics::GetModulationSystem(World))
+	{
+		ModSystem->OnEditPluginSettings(*this);
+	}
+}
+
+void USoundModulationSettings::PostDuplicate(EDuplicateMode::Type DuplicateMode)
+{
+	if (DuplicateMode == EDuplicateMode::Normal)
+	{
+		OnPostEditChange(GetWorld());
+	}
+}
+
+void USoundModulationSettings::PostEditChangeProperty(FPropertyChangedEvent& PropertyChangedEvent)
+{
+	OnPostEditChange(GetWorld());
+}
+
+void USoundModulationSettings::PostEditChangeChainProperty(FPropertyChangedChainEvent& PropertyChangedEvent)
+{
+	OnPostEditChange(GetWorld());
+}
+#endif // WITH_EDITOR
diff --git a/Engine/Plugins/Runtime/AudioModulation/Source/AudioModulation/Private/SoundModulationSettingsProxy.cpp b/Engine/Plugins/Runtime/AudioModulation/Source/AudioModulation/Private/SoundModulationSettingsProxy.cpp
new file mode 100644
--- /dev/null
+++ b/Engine/Plugins/Runtime/AudioModulation/Source/AudioModulation/Private/SoundModulationSettingsProxy.cpp
@@ -0,0 +1,36 @@
+// Copyright Epic Games, Inc. All Rights Reserved.
+#include "SoundModulationSettingsProxy.h"
+
+#include "AudioModulationSystem.h"
+#include "SoundModulationSettings.h"
+#include "SoundModulationTransform.h"
+
+
+namespace AudioModulation
+{
+	FModulationSettingsProxy::FModulationSettingsProxy()
+	{
+		Volume.OutputProxy.Transform.Curve = ESoundModulatorOutputCurve::Exp;
+		Volume.OutputProxy.Transform.Scalar = 0.5f;
+
+		Lowpass.OutputProxy.Transform.OutputMin = 20.0f;
+		Lowpass.OutputProxy.Transform.OutputMax = 20000.0f;
+		Lowpass.OutputProxy.Transform.Curve = ESoundModulatorOutputCurve::Exp;
+		Lowpass.OutputProxy.Operator = ESoundModulatorOperator::Min;
+
+		Highpass.DefaultInputValue = 0.0f;
+		Highpass.OutputProxy.Transform.OutputMin = 20.0f;
+		Highpass.OutputProxy.Transform.OutputMax = 20000.0f;
+		Highpass.OutputProxy.Transform.Curve = ESoundModulatorOutputCurve::Exp;
+		Highpass.OutputProxy.Operator = ESoundModulatorOperator::Max;
+	}
+
+	FModulationSettingsProxy::FModulationSettingsProxy(const USoundModulationSettings& Settings, FAudioModulationSystem& InModSystem)
+		: TModulatorProxyBase<uint32>(Settings.GetName(), Settings.GetUniqueID())
+		, Volume(Settings.Volume, InModSystem)
+		, Pitch(Settings.Pitch, InModSystem)
+		, Lowpass(Settings.Lowpass, InModSystem)
+		, Highpass(Settings.Highpass, InModSystem)
+	{
+	}
+} // namespace AudioModulation
diff --git a/Engine/Plugins/Runtime/AudioModulation/Source/AudioModulation/Private/SoundModulationSettingsProxy.h b/Engine/Plugins/Runtime/AudioModulation/Source/AudioModulation/Private/SoundModulationSettingsProxy.h
new file mode 100644
--- /dev/null
+++ b/Engine/Plugins/Runtime/AudioModulation/Source/AudioModulation/Private/SoundModulationSettingsProxy.h
@@ -0,0 +1,31 @@
+// Copyright Epic Games, Inc. All Rights Reserved.
+#pragma once
+
+#include "CoreMinimal.h"
+#include "IAudioModulation.h"
+#include "SoundModulationPatchProxy.h"
+#include "SoundModulationProxy.h"
+
+
+// Forward Declarations
+class USoundModulationSettings;
+
+namespace AudioModulation
+{
+	// Forward Declarations
+	class FAudioModulationSystem;
+
+	class FModulationSettingsProxy : public TModulatorProxyBase<uint32>
+	{
+	public:
+		FModulationSettingsProxy();
+		FModulationSettingsProxy(const USoundModulationSettings& Settings, FAudioModulationSystem& InModSystem);
+
+		FModulationPatchProxy Volume;
+		FModulationPatchProxy Pitch;
+		FModulationPatchProxy Lowpass;
+		FModulationPatchProxy Highpass;
+
+		TMap<FName, FModulationPatchProxy> Controls;
+	};
+} // namespace AudioModulation
\ No newline at end of file
diff --git a/Engine/Plugins/Runtime/AudioModulation/Source/AudioModulation/Private/SoundModulationValue.cpp b/Engine/Plugins/Runtime/AudioModulation/Source/AudioModulation/Private/SoundModulationValue.cpp
--- a/Engine/Plugins/Runtime/AudioModulation/Source/AudioModulation/Private/SoundModulationValue.cpp
+++ b/Engine/Plugins/Runtime/AudioModulation/Source/AudioModulation/Private/SoundModulationValue.cpp
@@ -4,13 +4,6 @@
 #include "SoundModulationProxy.h"
 
 
-namespace AudioModulation
-{
-	const FBusMixId InvalidBusMixId = INDEX_NONE;
-	const FBusId    InvalidBusId    = INDEX_NONE;
-	const FBusId    InvalidLFOId    = INDEX_NONE;
-} // namespace AudioModulation
-
 FSoundModulationValue::FSoundModulationValue()
 	: TargetValue(1.0f)
 	, AttackTime(0.1f)
diff --git a/Engine/Plugins/Runtime/AudioModulation/Source/AudioModulation/Private/SoundModulatorBase.cpp b/Engine/Plugins/Runtime/AudioModulation/Source/AudioModulation/Private/SoundModulatorBase.cpp
deleted file mode 100644
diff --git a/Engine/Plugins/Runtime/AudioModulation/Source/AudioModulation/Private/SoundModulatorLFO.cpp b/Engine/Plugins/Runtime/AudioModulation/Source/AudioModulation/Private/SoundModulatorLFO.cpp
--- a/Engine/Plugins/Runtime/AudioModulation/Source/AudioModulation/Private/SoundModulatorLFO.cpp
+++ b/Engine/Plugins/Runtime/AudioModulation/Source/AudioModulation/Private/SoundModulatorLFO.cpp
@@ -3,12 +3,29 @@
 
 #include "AudioDevice.h"
 #include "AudioModulation.h"
-#include "AudioModulationInternal.h"
 #include "AudioModulationLogging.h"
+#include "AudioModulationSystem.h"
 #include "Engine/World.h"
-#include "SoundModulationProxy.h"
+#include "SoundModulatorLFOProxy.h"
 
 
+USoundBusModulatorBase::USoundBusModulatorBase(const FObjectInitializer& ObjectInitializer)
+	: Super(ObjectInitializer)
+{
+}
+
+#if WITH_EDITOR
+void USoundBusModulatorBase::PostEditChangeProperty(FPropertyChangedEvent& InPropertyChangedEvent)
+{
+	if (FProperty* Property = InPropertyChangedEvent.Property)
+	{
+		AudioModulation::OnEditModulator(InPropertyChangedEvent, *this);
+	}
+
+	Super::PostEditChangeProperty(InPropertyChangedEvent);
+}
+#endif // WITH_EDITOR
+
 USoundBusModulatorLFO::USoundBusModulatorLFO(const FObjectInitializer& ObjectInitializer)
 	: Super(ObjectInitializer)
 	, Shape(ESoundModulatorLFOShape::Sine)
@@ -35,73 +52,9 @@ void USoundBusModulatorLFO::BeginDestroy()
 		check(AudioDevice->IsModulationPluginEnabled());
 		if (IAudioModulation* ModulationInterface = AudioDevice->ModulationInterface.Get())
 		{
-			auto ModulationImpl = static_cast<AudioModulation::FAudioModulation*>(ModulationInterface)->GetImpl();
-			check(ModulationImpl);
-			ModulationImpl->DeactivateLFO(*this);
+			auto ModSystem = static_cast<AudioModulation::FAudioModulation*>(ModulationInterface)->GetModulationSystem();
+			check(ModSystem);
+			ModSystem->DeactivateLFO(*this);
 		}
 	}
 }
-
-
-namespace AudioModulation
-{
-	FModulatorLFOProxy::FModulatorLFOProxy()
-		: Offset(0.0f)
-		, Value(1.0f)
-		, bBypass(0)
-	{
-		LFO.SetFrequency(1.0f);
-		LFO.Start();
-	}
-
-	FModulatorLFOProxy::FModulatorLFOProxy(const USoundBusModulatorLFO& InLFO, FAudioModulationImpl& InModulationImpl)
-		: TModulatorProxyRefType(InLFO.GetName(), InLFO.GetUniqueID(), InModulationImpl)
-		, Offset(InLFO.Offset)
-		, Value(1.0f)
-		, bBypass(InLFO.bBypass)
-	{
-		Init(InLFO);
-	}
-
-	FModulatorLFOProxy& FModulatorLFOProxy::operator =(const USoundBusModulatorLFO& InLFO)
-	{
-		Init(InLFO);
-		return *this;
-	}
-
-	float FModulatorLFOProxy::GetValue() const
-	{
-		return Value;
-	}
-
-	void FModulatorLFOProxy::Init(const USoundBusModulatorLFO& InLFO)
-	{
-		Offset = InLFO.Offset;
-		Value = 1.0f;
-		bBypass = InLFO.bBypass;
-
-		LFO.SetGain(InLFO.Amplitude);
-		LFO.SetFrequency(InLFO.Frequency);
-		LFO.SetMode(InLFO.bLooping ? Audio::ELFOMode::Type::Sync : Audio::ELFOMode::OneShot);
-
-		static_assert(static_cast<int32>(ESoundModulatorLFOShape::COUNT) == static_cast<int32>(Audio::ELFO::Type::NumLFOTypes), "LFOShape/ELFO Type mismatch");
-		LFO.SetType(static_cast<Audio::ELFO::Type>(InLFO.Shape));
-		LFO.Start();
-	}
-
-	bool FModulatorLFOProxy::IsBypassed() const
-	{
-		return bBypass != 0;
-	}
-
-	void FModulatorLFOProxy::Update(float InElapsed)
-	{
-		if (InElapsed > 0.0f && LFO.GetFrequency() > 0.0f)
-		{
-			const float SampleRate = 1.0f / InElapsed;
-			LFO.SetSampleRate(SampleRate);
-			LFO.Update();
-			Value = LFO.Generate() + Offset;
-		}
-	}
-} // namespace AudioModulation
diff --git a/Engine/Plugins/Runtime/AudioModulation/Source/AudioModulation/Private/SoundModulatorLFOProxy.cpp b/Engine/Plugins/Runtime/AudioModulation/Source/AudioModulation/Private/SoundModulatorLFOProxy.cpp
new file mode 100644
--- /dev/null
+++ b/Engine/Plugins/Runtime/AudioModulation/Source/AudioModulation/Private/SoundModulatorLFOProxy.cpp
@@ -0,0 +1,71 @@
+// Copyright Epic Games, Inc. All Rights Reserved.
+#include "SoundModulatorLFOProxy.h"
+
+#include "AudioModulation.h"
+#include "AudioModulationSystem.h"
+
+
+namespace AudioModulation
+{
+	const FLFOId InvalidLFOId = INDEX_NONE;
+
+	FModulatorLFOProxy::FModulatorLFOProxy()
+		: Offset(0.0f)
+		, Value(1.0f)
+		, bBypass(0)
+	{
+		LFO.SetFrequency(1.0f);
+		LFO.Start();
+	}
+
+	FModulatorLFOProxy::FModulatorLFOProxy(const USoundBusModulatorLFO& InLFO, FAudioModulationSystem& InModSystem)
+		: TModulatorProxyRefType(InLFO.GetName(), InLFO.GetUniqueID(), InModSystem)
+		, Offset(InLFO.Offset)
+		, Value(1.0f)
+		, bBypass(InLFO.bBypass)
+	{
+		Init(InLFO);
+	}
+
+	FModulatorLFOProxy& FModulatorLFOProxy::operator =(const USoundBusModulatorLFO& InLFO)
+	{
+		Init(InLFO);
+		return *this;
+	}
+
+	float FModulatorLFOProxy::GetValue() const
+	{
+		return Value;
+	}
+
+	void FModulatorLFOProxy::Init(const USoundBusModulatorLFO& InLFO)
+	{
+		Offset = InLFO.Offset;
+		Value = 1.0f;
+		bBypass = InLFO.bBypass;
+
+		LFO.SetGain(InLFO.Amplitude);
+		LFO.SetFrequency(InLFO.Frequency);
+		LFO.SetMode(InLFO.bLooping ? Audio::ELFOMode::Type::Sync : Audio::ELFOMode::OneShot);
+
+		static_assert(static_cast<int32>(ESoundModulatorLFOShape::COUNT) == static_cast<int32>(Audio::ELFO::Type::NumLFOTypes), "LFOShape/ELFO Type mismatch");
+		LFO.SetType(static_cast<Audio::ELFO::Type>(InLFO.Shape));
+		LFO.Start();
+	}
+
+	bool FModulatorLFOProxy::IsBypassed() const
+	{
+		return bBypass != 0;
+	}
+
+	void FModulatorLFOProxy::Update(float InElapsed)
+	{
+		if (InElapsed > 0.0f && LFO.GetFrequency() > 0.0f)
+		{
+			const float SampleRate = 1.0f / InElapsed;
+			LFO.SetSampleRate(SampleRate);
+			LFO.Update();
+			Value = LFO.Generate() + Offset;
+		}
+	}
+} // namespace AudioModulation
diff --git a/Engine/Plugins/Runtime/AudioModulation/Source/AudioModulation/Private/SoundModulatorLFOProxy.h b/Engine/Plugins/Runtime/AudioModulation/Source/AudioModulation/Private/SoundModulatorLFOProxy.h
new file mode 100644
--- /dev/null
+++ b/Engine/Plugins/Runtime/AudioModulation/Source/AudioModulation/Private/SoundModulatorLFOProxy.h
@@ -0,0 +1,42 @@
+// Copyright Epic Games, Inc. All Rights Reserved.
+#pragma once
+
+#include "SoundModulationProxy.h"
+#include "SoundModulatorLFO.h"
+
+
+namespace AudioModulation
+{
+	// Forward Declarations
+	class FAudioModulationSystem;
+
+	// Modulator Ids
+	using FLFOId = uint32;
+	extern const FLFOId InvalidLFOId;
+
+	class FModulatorLFOProxy;
+
+	using FLFOProxyMap = TMap<FLFOId, FModulatorLFOProxy>;
+	using FLFOHandle = TProxyHandle<FLFOId, FModulatorLFOProxy, USoundBusModulatorLFO>;
+
+	class FModulatorLFOProxy : public TModulatorProxyRefType<FLFOId, FModulatorLFOProxy, USoundBusModulatorLFO>
+	{
+	public:
+		FModulatorLFOProxy();
+		FModulatorLFOProxy(const USoundBusModulatorLFO& InLFO, FAudioModulationSystem& InModSystem);
+
+		FModulatorLFOProxy& operator =(const USoundBusModulatorLFO& InLFO);
+
+		float GetValue() const;
+		bool IsBypassed() const;
+		void Update(float InElapsed);
+
+	private:
+		void Init(const USoundBusModulatorLFO& InLFO);
+
+		Audio::FLFO LFO;
+		float Offset;
+		float Value;
+		uint8 bBypass : 1;
+	};
+} // namespace AudioModulation
\ No newline at end of file
diff --git a/Engine/Plugins/Runtime/AudioModulation/Source/AudioModulation/Public/AudioModulation.h b/Engine/Plugins/Runtime/AudioModulation/Source/AudioModulation/Public/AudioModulation.h
--- a/Engine/Plugins/Runtime/AudioModulation/Source/AudioModulation/Public/AudioModulation.h
+++ b/Engine/Plugins/Runtime/AudioModulation/Source/AudioModulation/Public/AudioModulation.h
@@ -4,8 +4,10 @@
 #include "CoreMinimal.h"
 #include "HAL/ThreadSafeBool.h"
 #include "IAudioExtensionPlugin.h"
+#include "IAudioModulation.h"
 #include "Modules/ModuleInterface.h"
 #include "SoundModulationPatch.h"
+#include "SoundModulationSettings.h"
 #include "Stats/Stats.h"
 
 
@@ -18,7 +20,7 @@ DECLARE_CYCLE_STAT_EXTERN(TEXT("Process Controls"), STAT_AudioModulationProcessC
 
 namespace AudioModulation
 {
-	class FAudioModulationImpl;
+	class FAudioModulationSystem;
 
 	class AUDIOMODULATION_API FAudioModulation : public IAudioModulation
 	{
@@ -30,6 +32,7 @@ namespace AudioModulation
 		virtual float CalculateInitialVolume(const USoundModulationPluginSourceSettingsBase& InSettingsBase) override;
 
 		virtual void Initialize(const FAudioPluginInitializationParams& InitializationParams) override;
+		virtual void OnBeginAudioRenderThreadUpdate() override;
 		virtual void OnInitSound(ISoundModulatable& Sound, const USoundModulationPluginSourceSettingsBase& Settings) override;
 		virtual void OnInitSource(const uint32 SourceId, const FName& AudioComponentUserId, const uint32 NumChannels, const USoundModulationPluginSourceSettingsBase& Settings) override;
 
@@ -43,16 +46,24 @@ namespace AudioModulation
 		virtual void OnReleaseSource(const uint32 SourceId) override;
 		virtual bool ProcessControls(const uint32 SourceId, FSoundModulationControls& Controls) override;
 		virtual void ProcessModulators(const float Elapsed) override;
+
+		virtual void UpdateModulator(const USoundModulatorBase& InModulator) override;
 		//~ End IAudioModulation implementation
 
 #if WITH_EDITOR
 		void OnEditPluginSettings(const USoundModulationPluginSourceSettingsBase& Settings);
 #endif // WITH_EDITOR
 
-		FAudioModulationImpl* GetImpl();
+		FAudioModulationSystem* GetModulationSystem();
+
+	protected:
+		virtual bool RegisterModulator(uint32 InParentId, const USoundModulatorBase& InModulatorBase) override;
+		virtual bool RegisterModulator(uint32 InParentId, Audio::FModulatorId InModulatorId) override;
+		virtual bool GetModulatorValue(const Audio::FModulatorHandle& ModulatorHandle, float& OutValue) override;
+		virtual void UnregisterModulator(const Audio::FModulatorHandle& InHandle) override;
 
 	private:
-		TUniquePtr<FAudioModulationImpl> Impl;
+		TUniquePtr<FAudioModulationSystem> ModSystem;
 	};
 } // namespace AudioModulation
 
diff --git a/Engine/Plugins/Runtime/AudioModulation/Source/AudioModulation/Public/AudioModulationSettings.h b/Engine/Plugins/Runtime/AudioModulation/Source/AudioModulation/Public/AudioModulationSettings.h
--- a/Engine/Plugins/Runtime/AudioModulation/Source/AudioModulation/Public/AudioModulationSettings.h
+++ b/Engine/Plugins/Runtime/AudioModulation/Source/AudioModulation/Public/AudioModulationSettings.h
@@ -10,7 +10,7 @@
 #include "AudioModulationSettings.generated.h"
 
 
-UCLASS(config = Game, defaultconfig, meta = (DisplayName = "Sound Modulation"))
+UCLASS(config = Game, defaultconfig, meta = (DisplayName = "Audio Modulation"))
 class AUDIOMODULATION_API UAudioModulationSettings : public UDeveloperSettings
 {
 	GENERATED_UCLASS_BODY()
diff --git a/Engine/Plugins/Runtime/AudioModulation/Source/AudioModulation/Public/AudioModulationStatics.h b/Engine/Plugins/Runtime/AudioModulation/Source/AudioModulation/Public/AudioModulationStatics.h
--- a/Engine/Plugins/Runtime/AudioModulation/Source/AudioModulation/Public/AudioModulationStatics.h
+++ b/Engine/Plugins/Runtime/AudioModulation/Source/AudioModulation/Public/AudioModulationStatics.h
@@ -17,7 +17,7 @@ class USoundControlBusMix;
 
 namespace AudioModulation
 {
-	class FAudioModulationImpl;
+	class FAudioModulationSystem;
 } // namespace AudioModulation
 
 
@@ -35,7 +35,7 @@ public:
 	/**
 	 * Returns modulation implementation associated with the provided world
 	 */
-	static AudioModulation::FAudioModulationImpl* GetModulationImpl(UWorld* World);
+	static AudioModulation::FAudioModulationSystem* GetModulationSystem(UWorld* World);
 
 	/** Activates a bus. Does nothing if an instance of the provided bus is already active
 	 * @param Bus - Bus to activate
@@ -254,6 +254,16 @@ public:
 		float								AttackTime	= -1.0f,
 		float								ReleaseTime = -1.0f);
 
+	/** Commits updates from a UObject definition of a bus mix to active instance in audio thread
+	 * (ignored if mix has not been activated).
+	 * @param Mix - Mix to update
+	 */
+	UFUNCTION(BlueprintCallable, Category = "Audio", DisplayName = "Update Control Bus Mix", meta = (
+		WorldContext = "WorldContextObject",
+		Keywords = "update set control bus mix modulation modulator")
+	)
+	static void UpdateMixFromObject(const UObject* WorldContextObject, USoundControlBusMix* Mix);
+
 	/** Commits updates from a UObject definition of a modulator (e.g. Bus, Bus Mix, LFO) to active instance in audio thread
 	 * (ignored if modulator type has not been activated).
 	 * @param Modulator - Modulator to update
diff --git a/Engine/Plugins/Runtime/AudioModulation/Source/AudioModulation/Public/SoundControlBus.h b/Engine/Plugins/Runtime/AudioModulation/Source/AudioModulation/Public/SoundControlBus.h
--- a/Engine/Plugins/Runtime/AudioModulation/Source/AudioModulation/Public/SoundControlBus.h
+++ b/Engine/Plugins/Runtime/AudioModulation/Source/AudioModulation/Public/SoundControlBus.h
@@ -1,31 +1,53 @@
 // Copyright Epic Games, Inc. All Rights Reserved.
 #pragma once
 
+#include "AudioDeviceManager.h"
+#include "AudioDevice.h"
 #include "CoreMinimal.h"
-#include "UObject/ObjectMacros.h"
-#include "UObject/Object.h"
-
-#include "SoundModulatorLFO.h"
+#include "IAudioModulation.h"
 #include "SoundModulationValue.h"
+#include "SoundModulatorLFO.h"
+#include "UObject/Object.h"
+#include "UObject/ObjectMacros.h"
 
 #include "SoundControlBus.generated.h"
 
 
-UENUM()
-enum class ESoundModulatorOperator : uint8
-{
-	/** Multiply all mix values together */
-	Multiply,
-
-	/** Take the lowest mix value */
-	Min,
+// Forward Declarations
+class USoundModulatorBase;
 
-	/** Take the highest mix value */
-	Max,
+struct FPropertyChangedEvent;
 
-	Count UMETA(Hidden),
-};
 
+#if WITH_EDITOR
+namespace AudioModulation
+{
+	static void OnEditModulator(FPropertyChangedEvent& InPropertyChangedEvent, USoundModulatorBase& InModulator)
+	{
+		check(InPropertyChangedEvent.Property);
+
+		if (!GEngine)
+		{
+			return;
+		}
+
+		FAudioDeviceManager* DeviceManager = GEngine->GetAudioDeviceManager();
+		if (!DeviceManager)
+		{
+			return;
+		}
+
+		const TArray<FAudioDevice*>& Devices = DeviceManager->GetAudioDevices();
+		for (FAudioDevice* Device : Devices)
+		{
+			if (Device && Device->IsModulationPluginEnabled() && Device->ModulationInterface.IsValid())
+			{
+				Device->ModulationInterface->UpdateModulator(InModulator);
+			}
+		}
+	}
+} // namespace AudioModulation
+#endif // WITH_EDITOR
 
 UCLASS(BlueprintType, hidecategories = Object, abstract, MinimalAPI)
 class USoundControlBusBase : public USoundModulatorBase
@@ -64,7 +86,7 @@ public:
 
 #if WITH_EDITOR
 	virtual void PostDuplicate(EDuplicateMode::Type DuplicateMode) override;
-	virtual void PostEditChangeProperty(FPropertyChangedEvent& PropertyChangedEvent) override;
+	virtual void PostEditChangeProperty(FPropertyChangedEvent& InPropertyChangedEvent) override;
 	virtual void PostInitProperties() override;
 	virtual void PostRename(UObject* OldOuter, const FName OldName) override;
 #endif // WITH_EDITOR
diff --git a/Engine/Plugins/Runtime/AudioModulation/Source/AudioModulation/Public/SoundControlBusMix.h b/Engine/Plugins/Runtime/AudioModulation/Source/AudioModulation/Public/SoundControlBusMix.h
--- a/Engine/Plugins/Runtime/AudioModulation/Source/AudioModulation/Public/SoundControlBusMix.h
+++ b/Engine/Plugins/Runtime/AudioModulation/Source/AudioModulation/Public/SoundControlBusMix.h
@@ -2,12 +2,11 @@
 #pragma once
 
 #include "CoreMinimal.h"
-#include "UObject/ObjectMacros.h"
-#include "UObject/Object.h"
-
+#include "IAudioModulation.h"
 #include "SoundControlBus.h"
-#include "SoundModulatorBase.h"
 #include "SoundModulationValue.h"
+#include "UObject/Object.h"
+#include "UObject/ObjectMacros.h"
 
 #include "SoundControlBusMix.generated.h"
 
@@ -34,7 +33,7 @@ struct FSoundControlBusMixChannel
 };
 
 UCLASS(config = Engine, autoexpandcategories = (Channel, Mix), editinlinenew, BlueprintType, MinimalAPI)
-class USoundControlBusMix : public USoundModulatorBase
+class USoundControlBusMix : public UObject
 {
 	GENERATED_UCLASS_BODY()
 
@@ -55,6 +54,10 @@ public:
 
 	virtual void BeginDestroy() override;
 
+#if WITH_EDITOR
+	virtual void PostEditChangeProperty(FPropertyChangedEvent& PropertyChangedEvent) override;
+#endif // WITH_EDITOR
+
 	/* Array of channels controlled by mix. */
 	UPROPERTY(EditAnywhere, Category = Mix, BlueprintReadOnly)
 	TArray<FSoundControlBusMixChannel> Channels;
diff --git a/Engine/Plugins/Runtime/AudioModulation/Source/AudioModulation/Public/SoundModulationPatch.h b/Engine/Plugins/Runtime/AudioModulation/Source/AudioModulation/Public/SoundModulationPatch.h
--- a/Engine/Plugins/Runtime/AudioModulation/Source/AudioModulation/Public/SoundModulationPatch.h
+++ b/Engine/Plugins/Runtime/AudioModulation/Source/AudioModulation/Public/SoundModulationPatch.h
@@ -11,15 +11,8 @@
 #include "SoundModulationPatch.generated.h"
 
 
-// Forward Declarations
-namespace AudioModulation
-{
-	struct FModulationInputProxy;
-} // namespace AudioModulation
-
-
 USTRUCT(BlueprintType)
-struct FSoundModulationOutputBase
+struct AUDIOMODULATION_API FSoundModulationOutputBase
 {
 	GENERATED_USTRUCT_BODY()
 
@@ -33,7 +26,7 @@ struct FSoundModulationOutputBase
 };
 
 USTRUCT(BlueprintType)
-struct FSoundModulationOutputFixedOperator : public FSoundModulationOutputBase
+struct AUDIOMODULATION_API FSoundModulationOutputFixedOperator : public FSoundModulationOutputBase
 {
 	GENERATED_USTRUCT_BODY()
 
@@ -50,7 +43,7 @@ public:
 };
 
 USTRUCT(BlueprintType)
-struct FSoundModulationOutput : public FSoundModulationOutputBase
+struct AUDIOMODULATION_API FSoundModulationOutput : public FSoundModulationOutputBase
 {
 	GENERATED_USTRUCT_BODY()
 
@@ -67,7 +60,7 @@ public:
 };
 
 USTRUCT(BlueprintType)
-struct FSoundModulationInputBase
+struct AUDIOMODULATION_API FSoundModulationInputBase
 {
 	GENERATED_USTRUCT_BODY()
 
@@ -86,7 +79,7 @@ struct FSoundModulationInputBase
 };
 
 USTRUCT(BlueprintType)
-struct FSoundVolumeModulationInput : public FSoundModulationInputBase
+struct AUDIOMODULATION_API FSoundVolumeModulationInput : public FSoundModulationInputBase
 {
 	GENERATED_USTRUCT_BODY()
 
@@ -100,7 +93,7 @@ struct FSoundVolumeModulationInput : public FSoundModulationInputBase
 };
 
 USTRUCT(BlueprintType)
-struct FSoundPitchModulationInput : public FSoundModulationInputBase
+struct AUDIOMODULATION_API FSoundPitchModulationInput : public FSoundModulationInputBase
 {
 	GENERATED_USTRUCT_BODY()
 
@@ -114,7 +107,7 @@ struct FSoundPitchModulationInput : public FSoundModulationInputBase
 };
 
 USTRUCT(BlueprintType)
-struct FSoundLPFModulationInput : public FSoundModulationInputBase
+struct AUDIOMODULATION_API FSoundLPFModulationInput : public FSoundModulationInputBase
 {
 	GENERATED_USTRUCT_BODY()
 
@@ -128,7 +121,7 @@ struct FSoundLPFModulationInput : public FSoundModulationInputBase
 };
 
 USTRUCT(BlueprintType)
-struct FSoundHPFModulationInput : public FSoundModulationInputBase
+struct AUDIOMODULATION_API FSoundHPFModulationInput : public FSoundModulationInputBase
 {
 	GENERATED_USTRUCT_BODY()
 
@@ -142,7 +135,7 @@ struct FSoundHPFModulationInput : public FSoundModulationInputBase
 };
 
 USTRUCT(BlueprintType)
-struct FSoundControlModulationInput : public FSoundModulationInputBase
+struct AUDIOMODULATION_API FSoundControlModulationInput : public FSoundModulationInputBase
 {
 	GENERATED_USTRUCT_BODY()
 
@@ -156,7 +149,7 @@ struct FSoundControlModulationInput : public FSoundModulationInputBase
 };
 
 USTRUCT(BlueprintType)
-struct FSoundModulationPatchBase
+struct AUDIOMODULATION_API FSoundModulationPatchBase
 {
 	GENERATED_USTRUCT_BODY()
 
@@ -182,7 +175,7 @@ struct FSoundModulationPatchBase
 };
 
 USTRUCT(BlueprintType)
-struct FSoundVolumeModulationPatch : public FSoundModulationPatchBase
+struct AUDIOMODULATION_API FSoundVolumeModulationPatch : public FSoundModulationPatchBase
 {
 	GENERATED_USTRUCT_BODY()
 
@@ -204,7 +197,7 @@ struct FSoundVolumeModulationPatch : public FSoundModulationPatchBase
 };
 
 USTRUCT(BlueprintType)
-struct FSoundPitchModulationPatch : public FSoundModulationPatchBase
+struct AUDIOMODULATION_API FSoundPitchModulationPatch : public FSoundModulationPatchBase
 {
 	GENERATED_USTRUCT_BODY()
 
@@ -226,7 +219,7 @@ struct FSoundPitchModulationPatch : public FSoundModulationPatchBase
 };
 
 USTRUCT(BlueprintType)
-struct FSoundLPFModulationPatch : public FSoundModulationPatchBase
+struct AUDIOMODULATION_API FSoundLPFModulationPatch : public FSoundModulationPatchBase
 {
 	GENERATED_USTRUCT_BODY()
 
@@ -248,7 +241,7 @@ struct FSoundLPFModulationPatch : public FSoundModulationPatchBase
 };
 
 USTRUCT(BlueprintType)
-struct FSoundHPFModulationPatch : public FSoundModulationPatchBase
+struct AUDIOMODULATION_API FSoundHPFModulationPatch : public FSoundModulationPatchBase
 {
 	GENERATED_USTRUCT_BODY()
 
@@ -270,7 +263,7 @@ struct FSoundHPFModulationPatch : public FSoundModulationPatchBase
 };
 
 USTRUCT(BlueprintType)
-struct FSoundControlModulationPatch : public FSoundModulationPatchBase
+struct AUDIOMODULATION_API FSoundControlModulationPatch : public FSoundModulationPatchBase
 {
 	GENERATED_USTRUCT_BODY()
 
@@ -291,34 +284,15 @@ struct FSoundControlModulationPatch : public FSoundModulationPatchBase
 	virtual const FSoundModulationOutputBase* GetOutput() const override { return static_cast<const FSoundModulationOutputBase*>(&Output); }
 };
 
-UCLASS(config = Engine, editinlinenew, BlueprintType, MinimalAPI, autoExpandCategories = (Volume, Pitch, Highpass, Lowpass))
-class USoundModulationSettings : public USoundModulationPluginSourceSettingsBase
+UCLASS(config = Engine, editinlinenew, BlueprintType)
+class AUDIOMODULATION_API USoundModulationPatch : public USoundModulatorBase
 {
 	GENERATED_UCLASS_BODY()
 
-public:
-	UPROPERTY(EditAnywhere, BlueprintReadWrite, Category = Volume)
-	FSoundVolumeModulationPatch Volume;
-
-	UPROPERTY(EditAnywhere, BlueprintReadWrite, Category = Pitch)
-	FSoundPitchModulationPatch Pitch;
-
-	UPROPERTY(EditAnywhere, BlueprintReadWrite, Category = Highpass)
-	FSoundHPFModulationPatch Highpass;
-
-	UPROPERTY(EditAnywhere, BlueprintReadWrite, Category = Lowpass)
-	FSoundLPFModulationPatch Lowpass;
-
-	// Array of named generic controls for use with modulateable parameters on source effects
-	// Properties hidden as Generic Control Modulation is still in development
-	UPROPERTY()
-	TArray<FSoundControlModulationPatch> Controls;
+	UPROPERTY(EditAnywhere, BlueprintReadWrite, Category = Settings)
+	FSoundControlModulationPatch PatchSettings;
 
 #if WITH_EDITOR
-	AUDIOMODULATION_API void OnPostEditChange(UWorld* World);
-
-	virtual void PostDuplicate(EDuplicateMode::Type DuplicateMode) override;
-	virtual void PostEditChangeProperty(FPropertyChangedEvent& PropertyChangedEvent) override;
-	virtual void PostEditChangeChainProperty(FPropertyChangedChainEvent& PropertyChangedEvent) override;
+	virtual void PostEditChangeProperty(FPropertyChangedEvent& InPropertyChangedEvent) override;
 #endif // WITH_EDITOR
 };
diff --git a/Engine/Plugins/Runtime/AudioModulation/Source/AudioModulation/Public/SoundModulationSettings.h b/Engine/Plugins/Runtime/AudioModulation/Source/AudioModulation/Public/SoundModulationSettings.h
new file mode 100644
--- /dev/null
+++ b/Engine/Plugins/Runtime/AudioModulation/Source/AudioModulation/Public/SoundModulationSettings.h
@@ -0,0 +1,38 @@
+// Copyright Epic Games, Inc. All Rights Reserved.
+#pragma once
+
+#include "CoreMinimal.h"
+#include "IAudioModulation.h"
+#include "SoundModulationPatch.h"
+#include "UObject/Object.h"
+#include "UObject/ObjectMacros.h"
+
+#include "SoundModulationSettings.generated.h"
+
+
+UCLASS(config = Engine, editinlinenew, BlueprintType, autoExpandCategories = (Volume, Pitch, Highpass, Lowpass))
+class AUDIOMODULATION_API USoundModulationSettings : public USoundModulationPluginSourceSettingsBase
+{
+	GENERATED_UCLASS_BODY()
+
+public:
+	UPROPERTY(EditAnywhere, BlueprintReadWrite, Category = Volume)
+	FSoundVolumeModulationPatch Volume;
+
+	UPROPERTY(EditAnywhere, BlueprintReadWrite, Category = Pitch)
+	FSoundPitchModulationPatch Pitch;
+
+	UPROPERTY(EditAnywhere, BlueprintReadWrite, Category = Highpass)
+	FSoundHPFModulationPatch Highpass;
+
+	UPROPERTY(EditAnywhere, BlueprintReadWrite, Category = Lowpass)
+	FSoundLPFModulationPatch Lowpass;
+
+#if WITH_EDITOR
+	void OnPostEditChange(UWorld* World);
+
+	virtual void PostDuplicate(EDuplicateMode::Type DuplicateMode) override;
+	virtual void PostEditChangeProperty(FPropertyChangedEvent& PropertyChangedEvent) override;
+	virtual void PostEditChangeChainProperty(FPropertyChangedChainEvent& PropertyChangedEvent) override;
+#endif // WITH_EDITOR
+};
diff --git a/Engine/Plugins/Runtime/AudioModulation/Source/AudioModulation/Public/SoundModulatorBase.h b/Engine/Plugins/Runtime/AudioModulation/Source/AudioModulation/Public/SoundModulatorBase.h
deleted file mode 100644
diff --git a/Engine/Plugins/Runtime/AudioModulation/Source/AudioModulation/Public/SoundModulatorLFO.h b/Engine/Plugins/Runtime/AudioModulation/Source/AudioModulation/Public/SoundModulatorLFO.h
--- a/Engine/Plugins/Runtime/AudioModulation/Source/AudioModulation/Public/SoundModulatorLFO.h
+++ b/Engine/Plugins/Runtime/AudioModulation/Source/AudioModulation/Public/SoundModulatorLFO.h
@@ -3,14 +3,14 @@
 
 #include "CoreMinimal.h"
 #include "DSP/LFO.h"
+#include "IAudioModulation.h"
+#include "SoundModulationValue.h"
 #include "UObject/ObjectMacros.h"
 #include "UObject/Object.h"
 
-#include "SoundModulatorBase.h"
-#include "SoundModulationValue.h"
-
 #include "SoundModulatorLFO.generated.h"
 
+
 UENUM(BlueprintType)
 enum class ESoundModulatorLFOShape : uint8
 {
@@ -25,6 +25,19 @@ enum class ESoundModulatorLFOShape : uint8
 	COUNT UMETA(Hidden)
 };
 
+/**
+ * Base class for modulators that manipulate control bus values
+ */
+UCLASS(hideCategories = Object, abstract, MinimalAPI)
+class USoundBusModulatorBase : public USoundModulatorBase
+{
+	GENERATED_UCLASS_BODY()
+
+#if WITH_EDITOR
+		virtual void PostEditChangeProperty(FPropertyChangedEvent& InPropertyChangedEvent) override;
+#endif // WITH_EDITOR
+};
+
 UCLASS(BlueprintType, hidecategories = Object, editinlinenew, MinimalAPI)
 class USoundBusModulatorLFO : public USoundBusModulatorBase
 {
diff --git a/Engine/Plugins/Runtime/AudioModulation/Source/AudioModulationEditor/Private/AssetTypeActions/AssetTypeActions_SoundModulationPatch.cpp b/Engine/Plugins/Runtime/AudioModulation/Source/AudioModulationEditor/Private/AssetTypeActions/AssetTypeActions_SoundModulationPatch.cpp
new file mode 100644
--- /dev/null
+++ b/Engine/Plugins/Runtime/AudioModulation/Source/AudioModulationEditor/Private/AssetTypeActions/AssetTypeActions_SoundModulationPatch.cpp
@@ -0,0 +1,25 @@
+// Copyright Epic Games, Inc. All Rights Reserved.
+#include "AssetTypeActions_SoundModulationPatch.h"
+
+#include "Editors/ModulationSettingsEditor.h"
+#include "SoundModulationPatch.h"
+
+
+#define LOCTEXT_NAMESPACE "AssetTypeActions"
+
+UClass* FAssetTypeActions_SoundModulationPatch::GetSupportedClass() const
+{
+	return USoundModulationPatch::StaticClass();
+}
+
+const TArray<FText>& FAssetTypeActions_SoundModulationPatch::GetSubMenus() const
+{
+	static const TArray<FText> SubMenus
+	{
+		FText(LOCTEXT("AssetSoundMixSubMenu", "Mix"))
+	};
+
+	return SubMenus;
+}
+
+#undef LOCTEXT_NAMESPACE
diff --git a/Engine/Plugins/Runtime/AudioModulation/Source/AudioModulationEditor/Private/AssetTypeActions/AssetTypeActions_SoundModulationPatch.h b/Engine/Plugins/Runtime/AudioModulation/Source/AudioModulationEditor/Private/AssetTypeActions/AssetTypeActions_SoundModulationPatch.h
new file mode 100644
--- /dev/null
+++ b/Engine/Plugins/Runtime/AudioModulation/Source/AudioModulationEditor/Private/AssetTypeActions/AssetTypeActions_SoundModulationPatch.h
@@ -0,0 +1,17 @@
+// Copyright Epic Games, Inc. All Rights Reserved.
+
+#pragma once
+
+#include "CoreMinimal.h"
+#include "AssetTypeActions_Base.h"
+
+class FAssetTypeActions_SoundModulationPatch : public FAssetTypeActions_Base
+{
+public:
+	// IAssetTypeActions Implementation
+	virtual FText GetName() const override { return NSLOCTEXT("AssetTypeActions", "AssetTypeActions_SoundModulationPatch", "Modulation Patch"); }
+	virtual FColor GetTypeColor() const override { return FColor(237, 237, 0); }
+	virtual const TArray<FText>& GetSubMenus() const override;
+	virtual UClass* GetSupportedClass() const override;
+	virtual uint32 GetCategories() override { return EAssetTypeCategories::Sounds; }
+};
diff --git a/Engine/Plugins/Runtime/AudioModulation/Source/AudioModulationEditor/Private/AssetTypeActions/AssetTypeActions_SoundModulationSettings.cpp b/Engine/Plugins/Runtime/AudioModulation/Source/AudioModulationEditor/Private/AssetTypeActions/AssetTypeActions_SoundModulationSettings.cpp
--- a/Engine/Plugins/Runtime/AudioModulation/Source/AudioModulationEditor/Private/AssetTypeActions/AssetTypeActions_SoundModulationSettings.cpp
+++ b/Engine/Plugins/Runtime/AudioModulation/Source/AudioModulationEditor/Private/AssetTypeActions/AssetTypeActions_SoundModulationSettings.cpp
@@ -2,7 +2,7 @@
 #include "AssetTypeActions_SoundModulationSettings.h"
 
 #include "Editors/ModulationSettingsEditor.h"
-#include "SoundModulationPatch.h"
+#include "SoundModulationSettings.h"
 
 
 #define LOCTEXT_NAMESPACE "AssetTypeActions"
diff --git a/Engine/Plugins/Runtime/AudioModulation/Source/AudioModulationEditor/Private/AudioModulationEditor.cpp b/Engine/Plugins/Runtime/AudioModulation/Source/AudioModulationEditor/Private/AudioModulationEditor.cpp
--- a/Engine/Plugins/Runtime/AudioModulation/Source/AudioModulationEditor/Private/AudioModulationEditor.cpp
+++ b/Engine/Plugins/Runtime/AudioModulation/Source/AudioModulationEditor/Private/AudioModulationEditor.cpp
@@ -1,21 +1,29 @@
 // Copyright Epic Games, Inc. All Rights Reserved.
 #include "AudioModulationEditor.h"
 
-#include "AssetTypeActions/AssetTypeActions_SoundModulationSettings.h"
 #include "AssetTypeActions/AssetTypeActions_SoundControlBus.h"
 #include "AssetTypeActions/AssetTypeActions_SoundControlBusMix.h"
+#include "AssetTypeActions/AssetTypeActions_SoundModulationPatch.h"
+#include "AssetTypeActions/AssetTypeActions_SoundModulationSettings.h"
 #include "AssetTypeActions/AssetTypeActions_SoundModulatorLFO.h"
 #include "Editors/ModulationSettingsCurveEditorViewStacked.h"
+#include "Framework/Commands/UIAction.h"
+#include "Framework/MultiBox/MultiBoxBuilder.h"
+#include "Framework/MultiBox/MultiBoxExtender.h"
+#include "ICurveEditorModule.h"
+#include "Internationalization/Internationalization.h"
 #include "Layouts/SoundControlModulationPatchLayout.h"
-#include "Layouts/SoundModulationParameterLayout.h"
 #include "Layouts/SoundModulationTransformLayout.h"
+#include "LevelEditor.h"
 #include "SoundModulationTransform.h"
-#include "ICurveEditorModule.h"
+#include "Templates/SharedPointer.h"
+#include "Textures/SlateIcon.h"
+#include "UObject/UObjectIterator.h"
 
 
-namespace
+namespace AudioModulationEditor
 {
-	static const FName AssetToolsName = TEXT("AssetTools");
+	static const FName ToolName = TEXT("AssetTools");
 
 	template <typename T>
 	void AddAssetAction(IAssetTools& AssetTools, TArray<TSharedPtr<FAssetTypeActions_Base>>& AssetArray)
@@ -25,7 +33,7 @@ namespace
 		AssetTools.RegisterAssetTypeActions(AssetAction.ToSharedRef());
 		AssetArray.Add(AssetActionBase);
 	}
-} // namespace <>
+} // namespace AudioModulationEditor
 
 FAudioModulationEditorModule::FAudioModulationEditorModule()
 {
@@ -64,16 +72,17 @@ void FAudioModulationEditorModule::StartupModule()
 	ModulationSettingsMenuExtensibilityManager = MakeShared<FExtensibilityManager>();
 
 	// Register the audio editor asset type actions
-	IAssetTools& AssetTools = FModuleManager::LoadModuleChecked<FAssetToolsModule>(AssetToolsName).Get();
-
-	AddAssetAction<FAssetTypeActions_SoundVolumeControlBus>(AssetTools, AssetActions);
-	AddAssetAction<FAssetTypeActions_SoundPitchControlBus>(AssetTools, AssetActions);
-	AddAssetAction<FAssetTypeActions_SoundHPFControlBus>(AssetTools, AssetActions);
-	AddAssetAction<FAssetTypeActions_SoundLPFControlBus>(AssetTools, AssetActions);
-	AddAssetAction<FAssetTypeActions_SoundControlBus>(AssetTools, AssetActions);
-	AddAssetAction<FAssetTypeActions_SoundControlBusMix>(AssetTools, AssetActions);
-	AddAssetAction<FAssetTypeActions_SoundModulatorLFO>(AssetTools, AssetActions);
-	AddAssetAction<FAssetTypeActions_SoundModulationSettings>(AssetTools, AssetActions);
+	IAssetTools& AssetTools = FModuleManager::LoadModuleChecked<FAssetToolsModule>(AudioModulationEditor::ToolName).Get();
+
+	AudioModulationEditor::AddAssetAction<FAssetTypeActions_SoundVolumeControlBus>(AssetTools, AssetActions);
+	AudioModulationEditor::AddAssetAction<FAssetTypeActions_SoundPitchControlBus>(AssetTools, AssetActions);
+	AudioModulationEditor::AddAssetAction<FAssetTypeActions_SoundHPFControlBus>(AssetTools, AssetActions);
+	AudioModulationEditor::AddAssetAction<FAssetTypeActions_SoundLPFControlBus>(AssetTools, AssetActions);
+	AudioModulationEditor::AddAssetAction<FAssetTypeActions_SoundControlBus>(AssetTools, AssetActions);
+	AudioModulationEditor::AddAssetAction<FAssetTypeActions_SoundControlBusMix>(AssetTools, AssetActions);
+	AudioModulationEditor::AddAssetAction<FAssetTypeActions_SoundModulatorLFO>(AssetTools, AssetActions);
+	AudioModulationEditor::AddAssetAction<FAssetTypeActions_SoundModulationPatch>(AssetTools, AssetActions);
+	AudioModulationEditor::AddAssetAction<FAssetTypeActions_SoundModulationSettings>(AssetTools, AssetActions);
 
 	SetIcon(TEXT("SoundVolumeControlBus"));
 	SetIcon(TEXT("SoundPitchControlBus"));
@@ -103,9 +112,6 @@ void FAudioModulationEditorModule::RegisterCustomPropertyLayouts()
 	PropertyModule.RegisterCustomPropertyTypeLayout("SoundModulationOutputTransform",
 		FOnGetPropertyTypeCustomizationInstance::CreateStatic(
 			&FSoundModulationOutputTransformLayoutCustomization::MakeInstance));
-// 	PropertyModule.RegisterCustomPropertyTypeLayout("SoundModulationParameter",
-// 		FOnGetPropertyTypeCustomizationInstance::CreateStatic(
-// 			&FSoundModulationParameterLayoutCustomization::MakeInstance));
 	PropertyModule.RegisterCustomPropertyTypeLayout("SoundVolumeModulationPatch",
 		FOnGetPropertyTypeCustomizationInstance::CreateStatic(
 			&FSoundVolumeModulationPatchLayoutCustomization::MakeInstance));
@@ -118,9 +124,6 @@ void FAudioModulationEditorModule::RegisterCustomPropertyLayouts()
 	PropertyModule.RegisterCustomPropertyTypeLayout("SoundLPFModulationPatch",
 		FOnGetPropertyTypeCustomizationInstance::CreateStatic(
 			&FSoundLPFModulationPatchLayoutCustomization::MakeInstance));
-	PropertyModule.RegisterCustomPropertyTypeLayout("SoundControlModulationPatch",
-		FOnGetPropertyTypeCustomizationInstance::CreateStatic(
-			&FSoundControlModulationPatchLayoutCustomization::MakeInstance));
 }
 
 void FAudioModulationEditorModule::ShutdownModule()
@@ -128,9 +131,9 @@ void FAudioModulationEditorModule::ShutdownModule()
 	ModulationSettingsToolBarExtensibilityManager.Reset();
 	ModulationSettingsMenuExtensibilityManager.Reset();
 
-	if (FModuleManager::Get().IsModuleLoaded(AssetToolsName))
+	if (FModuleManager::Get().IsModuleLoaded(AudioModulationEditor::ToolName))
 	{
-		IAssetTools& AssetTools = FModuleManager::GetModuleChecked<FAssetToolsModule>(AssetToolsName).Get();
+		IAssetTools& AssetTools = FModuleManager::GetModuleChecked<FAssetToolsModule>(AudioModulationEditor::ToolName).Get();
 		for (TSharedPtr<FAssetTypeActions_Base>& AssetAction : AssetActions)
 		{
 			AssetTools.UnregisterAssetTypeActions(AssetAction.ToSharedRef());
diff --git a/Engine/Plugins/Runtime/AudioModulation/Source/AudioModulationEditor/Private/Editors/ModulationSettingsCurveEditorViewStacked.cpp b/Engine/Plugins/Runtime/AudioModulation/Source/AudioModulationEditor/Private/Editors/ModulationSettingsCurveEditorViewStacked.cpp
--- a/Engine/Plugins/Runtime/AudioModulation/Source/AudioModulationEditor/Private/Editors/ModulationSettingsCurveEditorViewStacked.cpp
+++ b/Engine/Plugins/Runtime/AudioModulation/Source/AudioModulationEditor/Private/Editors/ModulationSettingsCurveEditorViewStacked.cpp
@@ -10,6 +10,7 @@
 #include "Fonts/FontMeasure.h"
 #include "SCurveEditorPanel.h"
 #include "SoundModulationPatch.h"
+#include "SoundModulationSettings.h"
 #include "Widgets/Text/STextBlock.h"
 
 
diff --git a/Engine/Plugins/Runtime/AudioModulation/Source/AudioModulationEditor/Private/Editors/ModulationSettingsEditor.cpp b/Engine/Plugins/Runtime/AudioModulation/Source/AudioModulationEditor/Private/Editors/ModulationSettingsEditor.cpp
--- a/Engine/Plugins/Runtime/AudioModulation/Source/AudioModulationEditor/Private/Editors/ModulationSettingsEditor.cpp
+++ b/Engine/Plugins/Runtime/AudioModulation/Source/AudioModulationEditor/Private/Editors/ModulationSettingsEditor.cpp
@@ -15,7 +15,6 @@
 #include "Framework/MultiBox/MultiBoxBuilder.h"
 #include "Framework/MultiBox/MultiBoxExtender.h"
 #include "IDetailsView.h"
-#include "ModulationSettingsEditor.h"
 #include "Modules/ModuleManager.h"
 #include "PropertyEditorModule.h"
 #include "RichCurveEditorModel.h"
@@ -30,6 +29,7 @@
 #include "Widgets/Layout/SBorder.h"
 #include "Widgets/SFrameRatePicker.h"
 #include "ModulationSettingsCurveEditorViewStacked.h"
+#include "SoundModulationSettings.h"
 
 
 #define LOCTEXT_NAMESPACE "ModulationSettingsEditor"
@@ -236,66 +236,11 @@ void FModulationSettingsEditor::SetCurveAtOrderIndex(int32 InCurveIndex, FRichCu
 			TUniquePtr<FModCurveEditorModel> NewCurve = MakeUnique<FModCurveEditorModel>(InRichCurve, GetEditingObject(), CurveOutput, InSource, InSharedCurve);
 			CurveModels[InCurveIndex] = CurveEditor->AddCurve(MoveTemp(NewCurve));
 		}
-		else
-		{
-			const USoundModulationSettings* ModSettings = CastChecked<USoundModulationSettings>(GetEditingObject());
-			const int32 ControlIndex = GetControlIndex(InCurveIndex);
-			const FName Name = ModSettings->Controls[ControlIndex].Control;
-			TUniquePtr<FModCurveEditorModel> NewCurve = MakeUnique<FModCurveEditorModel>(InRichCurve, GetEditingObject(), Name, InSource, InSharedCurve);
-			CurveModels[InCurveIndex] = CurveEditor->AddCurve(MoveTemp(NewCurve));
-		}
 	}
 
 	CurveEditor->PinCurve(CurveModels[InCurveIndex]);
 }
 
-int32 FModulationSettingsEditor::GetControlModulationCount() const
-{
-	USoundModulationSettings* Settings = Cast<USoundModulationSettings>(GetEditingObject());
-	if (!Settings)
-	{
-		return 0;
-	}
-
-	return Settings->Controls.Num();
-}
-
-int32 FModulationSettingsEditor::GetCurveCount() const
-{
-	USoundModulationSettings* Settings = Cast<USoundModulationSettings>(GetEditingObject());
-	if (!Settings)
-	{
-		return 0;
-	}
-
-	return static_cast<int32>(EModSettingsEditorCurveOutput::Count) - 1 + Settings->Controls.Num();
-}
-
-int32 FModulationSettingsEditor::GetCurveOrderIndex(EModSettingsEditorCurveOutput InCurveOutput, int32 InControlIndex) const
-{
-	switch (InCurveOutput)
-	{
-	case EModSettingsEditorCurveOutput::Volume:
-	case EModSettingsEditorCurveOutput::Pitch:
-	case EModSettingsEditorCurveOutput::Highpass:
-	case EModSettingsEditorCurveOutput::Lowpass:
-	{
-		return static_cast<int32>(InCurveOutput);
-	}
-	break;
-
-	case EModSettingsEditorCurveOutput::Control:
-	{
-		return static_cast<int32>(InCurveOutput) + InControlIndex;
-	}
-	break;
-
-	default:
-		static_assert(static_cast<int32>(EModSettingsEditorCurveOutput::Count) == 5, "Possible missing case coverage for enum EModSettingsEditorCurveOutput");
-		return -1;
-	}
-}
-
 int32 FModulationSettingsEditor::GetControlIndex(int32 InCurveIndex) const
 {
 	const int32 ControlIndexOffset = static_cast<int32>(EModSettingsEditorCurveOutput::Control);
@@ -307,24 +252,6 @@ int32 FModulationSettingsEditor::GetControlIndex(int32 InCurveIndex) const
 	return InCurveIndex - ControlIndexOffset;
 }
 
-FName FModulationSettingsEditor::GetControlName(int32 InCurveIndex) const
-{
-	const int32 ControlIndexOffset = static_cast<int32>(EModSettingsEditorCurveOutput::Control);
-	if (InCurveIndex < ControlIndexOffset)
-	{
-		return NAME_None;
-	}
-
-	USoundModulationSettings* Settings = Cast<USoundModulationSettings>(GetEditingObject());
-	if (!Settings)
-	{
-		return NAME_None;
-	}
-
-	const int32 ControlIndex = InCurveIndex - static_cast<int32>(EModSettingsEditorCurveOutput::Control);
-	return Settings->Controls[ControlIndex].Control;
-}
-
 bool FModulationSettingsEditor::GetIsBypassed(int32 InCurveIndex) const
 {
 	if (USoundModulationSettings* Settings = Cast<USoundModulationSettings>(GetEditingObject()))
@@ -379,8 +306,7 @@ FSoundModulationOutputBase* FModulationSettingsEditor::FindModulationOutput(int3
 
 	if (InCurveIndex >= static_cast<int32>(EModSettingsEditorCurveOutput::Control))
 	{
-		const int32 ControlIndex = InCurveIndex - static_cast<int32>(EModSettingsEditorCurveOutput::Control);
-		return &Settings->Controls[ControlIndex].Output;
+		return nullptr;
 	}
 
 	const EModSettingsEditorCurveOutput CurveOutput = static_cast<EModSettingsEditorCurveOutput>(InCurveIndex);
@@ -566,18 +492,11 @@ void FModulationSettingsEditor::UpdateCurves()
 
 	CurveEditor->RemoveAllCurves();
 
-	int32 CurveCount = GetCurveCount();
+	const int32 CurveCount = static_cast<int32>(EModSettingsEditorCurveOutput::Count) - 1;
 	for (int32 i = 0; i < CurveCount; ++i)
 	{
 		UpdateCurve(i);
 	}
-
-	// Remove old curves from editor
-	CurveCount = GetCurveCount();
-	for (int32 i = CurveCount; i < CurveModels.Num(); ++i)
-	{
-		CurveEditor->RemoveCurve(CurveModels[i]);
-	}
 }
 
 void FModulationSettingsEditor::PostRedo(bool bSuccess)
diff --git a/Engine/Plugins/Runtime/AudioModulation/Source/AudioModulationEditor/Private/Editors/ModulationSettingsEditor.h b/Engine/Plugins/Runtime/AudioModulation/Source/AudioModulationEditor/Private/Editors/ModulationSettingsEditor.h
--- a/Engine/Plugins/Runtime/AudioModulation/Source/AudioModulationEditor/Private/Editors/ModulationSettingsEditor.h
+++ b/Engine/Plugins/Runtime/AudioModulation/Source/AudioModulationEditor/Private/Editors/ModulationSettingsEditor.h
@@ -58,16 +58,8 @@ private:
 	/**	Spawns the tab allowing for editing/viewing the output curve(s) */
 	TSharedRef<SDockTab> SpawnTab_Properties(const FSpawnTabArgs& Args);
 
-	int32 GetControlModulationCount() const;
-
-	int32 GetCurveOrderIndex(EModSettingsEditorCurveOutput InCurveOutput, int32 InControlIndex) const;
-
 	int32 GetControlIndex(int32 InCurveIndex) const;
 
-	FName GetControlName(int32 InCurveIndex) const;
-
-	int32 GetCurveCount() const;
-
 	bool GetIsBypassed(int32 InCurveIndex) const;
 
 	/** Get the orientation for the snap value controls. */
diff --git a/Engine/Plugins/Runtime/AudioModulation/Source/AudioModulationEditor/Private/Factories/SoundModulationPatchFactory.cpp b/Engine/Plugins/Runtime/AudioModulation/Source/AudioModulationEditor/Private/Factories/SoundModulationPatchFactory.cpp
new file mode 100644
--- /dev/null
+++ b/Engine/Plugins/Runtime/AudioModulation/Source/AudioModulationEditor/Private/Factories/SoundModulationPatchFactory.cpp
@@ -0,0 +1,18 @@
+// Copyright Epic Games, Inc. All Rights Reserved.
+#include "SoundModulationPatchFactory.h"
+#include "SoundModulationPatch.h"
+
+
+USoundModulationPatchFactory::USoundModulationPatchFactory(const FObjectInitializer& ObjectInitializer)
+	: Super(ObjectInitializer)
+{
+	SupportedClass = USoundModulationPatch::StaticClass();
+	bCreateNew = true;
+	bEditorImport = false;
+	bEditAfterNew = true;
+}
+
+UObject* USoundModulationPatchFactory::FactoryCreateNew(UClass* Class, UObject* InParent, FName Name, EObjectFlags Flags, UObject* Context, FFeedbackContext* Warn)
+{
+	return NewObject<USoundModulationPatch>(InParent, Name, Flags);
+}
diff --git a/Engine/Plugins/Runtime/AudioModulation/Source/AudioModulationEditor/Private/Factories/SoundModulationPatchFactory.h b/Engine/Plugins/Runtime/AudioModulation/Source/AudioModulationEditor/Private/Factories/SoundModulationPatchFactory.h
new file mode 100644
--- /dev/null
+++ b/Engine/Plugins/Runtime/AudioModulation/Source/AudioModulationEditor/Private/Factories/SoundModulationPatchFactory.h
@@ -0,0 +1,19 @@
+// Copyright Epic Games, Inc. All Rights Reserved.
+#pragma once
+
+#include "CoreMinimal.h"
+#include "UObject/ObjectMacros.h"
+#include "Factories/Factory.h"
+
+#include "SoundModulationPatchFactory.generated.h"
+
+
+UCLASS(hidecategories=Object, MinimalAPI)
+class USoundModulationPatchFactory : public UFactory
+{
+	GENERATED_UCLASS_BODY()
+
+	// Begin UFactory Interface
+	virtual UObject* FactoryCreateNew(UClass* Class,UObject* InParent,FName Name,EObjectFlags Flags,UObject* Context,FFeedbackContext* Warn) override;
+	// Begin UFactory Interface	
+};
diff --git a/Engine/Plugins/Runtime/AudioModulation/Source/AudioModulationEditor/Private/Factories/SoundModulationSettingsFactory.cpp b/Engine/Plugins/Runtime/AudioModulation/Source/AudioModulationEditor/Private/Factories/SoundModulationSettingsFactory.cpp
--- a/Engine/Plugins/Runtime/AudioModulation/Source/AudioModulationEditor/Private/Factories/SoundModulationSettingsFactory.cpp
+++ b/Engine/Plugins/Runtime/AudioModulation/Source/AudioModulationEditor/Private/Factories/SoundModulationSettingsFactory.cpp
@@ -1,6 +1,7 @@
 // Copyright Epic Games, Inc. All Rights Reserved.
 #include "SoundModulationSettingsFactory.h"
-#include "SoundModulationPatch.h"
+
+#include "SoundModulationSettings.h"
 
 
 USoundModulationSettingsFactory::USoundModulationSettingsFactory(const FObjectInitializer& ObjectInitializer)
diff --git a/Engine/Plugins/Runtime/AudioModulation/Source/AudioModulationEditor/Private/Layouts/SoundControlModulationPatchLayout.cpp b/Engine/Plugins/Runtime/AudioModulation/Source/AudioModulationEditor/Private/Layouts/SoundControlModulationPatchLayout.cpp
--- a/Engine/Plugins/Runtime/AudioModulation/Source/AudioModulationEditor/Private/Layouts/SoundControlModulationPatchLayout.cpp
+++ b/Engine/Plugins/Runtime/AudioModulation/Source/AudioModulationEditor/Private/Layouts/SoundControlModulationPatchLayout.cpp
@@ -6,7 +6,7 @@
 #include "DetailLayoutBuilder.h"
 #include "DetailWidgetRow.h"
 #include "EditorStyleSet.h"
-#include "IAudioExtensionPlugin.h"
+#include "IAudioModulation.h"
 #include "SoundModulationPatch.h"
 #include "Widgets/Input/SButton.h"
 #include "Widgets/Text/STextBlock.h"
@@ -89,39 +89,4 @@ TAttribute<EVisibility> FSoundLPFModulationPatchLayoutCustomization::CustomizeCo
 	AddPatchProperties<FSoundLPFModulationPatch>(VisibilityAttribute, PropertyHandles, ChildBuilder);
 	return VisibilityAttribute;
 }
-
-TAttribute<EVisibility> FSoundControlModulationPatchLayoutCustomization::CustomizeControl(TMap<FName, TSharedPtr<IPropertyHandle>> &PropertyHandles, IDetailChildrenBuilder &ChildBuilder)
-{
-// Properties hidden as Generic Control Modulation is still in development
-// 	FName SetControlName;
-// 
-// 	TSharedPtr<IPropertyHandle> ControlProperty = PropertyHandles.FindChecked(GET_MEMBER_NAME_CHECKED(FSoundControlModulationPatch, Control));
-// 	ControlProperty->GetValue(SetControlName);
-// 
-// 	TSharedPtr<IPropertyHandle>ValueHandle = PropertyHandles.FindChecked(GET_MEMBER_NAME_CHECKED(FSoundControlModulationPatch, DefaultInputValue));
-// 
-// 	ChildBuilder.AddCustomRow(FText::FromName(SetControlName))
-// 	.NameContent()
-// 	[
-// 		SAssignNew(ControlComboBox, SSoundModulationControlComboBox)
-// 		.InitControlName(SetControlName)
-// 		.ControlNameProperty(ControlProperty)
-// 	]
-// 	.ValueContent()
-// 	.MinDesiredWidth(120)
-// 	[
-// 		ValueHandle->CreatePropertyValueWidget()
-// 	];
-// 
-// 	TSharedPtr<IPropertyHandle>InputsHandle = PropertyHandles.FindChecked(GET_MEMBER_NAME_CHECKED(FSoundControlModulationPatch, Inputs));
-// 	ChildBuilder.AddProperty(InputsHandle.ToSharedRef());
-// 
-// 	TSharedPtr<IPropertyHandle>OutputHandle = PropertyHandles.FindChecked(GET_MEMBER_NAME_CHECKED(FSoundControlModulationPatch, Output));
-// 	ChildBuilder.AddProperty(OutputHandle.ToSharedRef());
-
-	return TAttribute<EVisibility>::Create([this]()
-	{
-		return EVisibility::Hidden;
-	});
-}
 #undef LOCTEXT_NAMESPACE
diff --git a/Engine/Plugins/Runtime/AudioModulation/Source/AudioModulationEditor/Private/Layouts/SoundControlModulationPatchLayout.h b/Engine/Plugins/Runtime/AudioModulation/Source/AudioModulationEditor/Private/Layouts/SoundControlModulationPatchLayout.h
--- a/Engine/Plugins/Runtime/AudioModulation/Source/AudioModulationEditor/Private/Layouts/SoundControlModulationPatchLayout.h
+++ b/Engine/Plugins/Runtime/AudioModulation/Source/AudioModulationEditor/Private/Layouts/SoundControlModulationPatchLayout.h
@@ -9,7 +9,6 @@
 #include "Layout/Visibility.h"
 #include "Misc/Attribute.h"
 #include "PropertyHandle.h"
-#include "SoundModulationControlComboBox.h"
 
 
 // Forward Declarations
@@ -90,16 +89,3 @@ public:
 	virtual TAttribute<EVisibility> CustomizeControl(TMap<FName, TSharedPtr<IPropertyHandle>>& PropertyHandles, IDetailChildrenBuilder& ChildBuilder) override;
 };
 
-class FSoundControlModulationPatchLayoutCustomization : public FSoundModulationPatchLayoutCustomization
-{
-public:
-	static TSharedRef<IPropertyTypeCustomization> MakeInstance()
-	{
-		return MakeShared<FSoundControlModulationPatchLayoutCustomization>();
-	}
-
-	virtual TAttribute<EVisibility> CustomizeControl(TMap<FName, TSharedPtr<IPropertyHandle>>& PropertyHandles, IDetailChildrenBuilder& ChildBuilder) override;
-
-protected:
-	TSharedPtr<SSearchableComboBox> ControlComboBox;
-};
diff --git a/Engine/Plugins/Runtime/AudioModulation/Source/AudioModulationEditor/Private/Layouts/SoundModulationControlComboBox.cpp b/Engine/Plugins/Runtime/AudioModulation/Source/AudioModulationEditor/Private/Layouts/SoundModulationControlComboBox.cpp
deleted file mode 100644
diff --git a/Engine/Plugins/Runtime/AudioModulation/Source/AudioModulationEditor/Private/Layouts/SoundModulationControlComboBox.h b/Engine/Plugins/Runtime/AudioModulation/Source/AudioModulationEditor/Private/Layouts/SoundModulationControlComboBox.h
deleted file mode 100644
diff --git a/Engine/Plugins/Runtime/AudioModulation/Source/AudioModulationEditor/Private/Layouts/SoundModulationParameterLayout.cpp b/Engine/Plugins/Runtime/AudioModulation/Source/AudioModulationEditor/Private/Layouts/SoundModulationParameterLayout.cpp
deleted file mode 100644
diff --git a/Engine/Plugins/Runtime/AudioModulation/Source/AudioModulationEditor/Private/Layouts/SoundModulationParameterLayout.h b/Engine/Plugins/Runtime/AudioModulation/Source/AudioModulationEditor/Private/Layouts/SoundModulationParameterLayout.h
deleted file mode 100644
diff --git a/Engine/Plugins/Runtime/Synthesis/Source/Synthesis/Synthesis.Build.cs b/Engine/Plugins/Runtime/Synthesis/Source/Synthesis/Synthesis.Build.cs
--- a/Engine/Plugins/Runtime/Synthesis/Source/Synthesis/Synthesis.Build.cs
+++ b/Engine/Plugins/Runtime/Synthesis/Source/Synthesis/Synthesis.Build.cs
@@ -4,22 +4,30 @@ namespace UnrealBuildTool.Rules
 {
 	public class Synthesis : ModuleRules
 	{
-        public Synthesis(ReadOnlyTargetRules Target) : base(Target)
+		public Synthesis(ReadOnlyTargetRules Target) : base(Target)
 		{
-            PrivateDependencyModuleNames.AddRange(
+			OptimizeCode = CodeOptimization.Never;
+
+			PublicDependencyModuleNames.AddRange(
 				new string[] {
-                    "Core",
+					"Core",
 					"CoreUObject",
 					"Engine",
+					"AudioExtensions",
+				}
+			);
+
+			PrivateDependencyModuleNames.AddRange(
+				new string[] {
 					"AudioMixer",
 					"SignalProcessing",
-                    "UMG",
-                    "Slate",
-                    "SlateCore",
-                    "InputCore",
-                    "Projects"
-                }
-            );
+					"UMG",
+					"Slate",
+					"SlateCore",
+					"InputCore",
+					"Projects"
+				}
+			);
 		}
 	}
 }
\ No newline at end of file
diff --git a/Engine/Source/Runtime/AudioExtensions/AudioExtensions.Build.cs b/Engine/Source/Runtime/AudioExtensions/AudioExtensions.Build.cs
--- a/Engine/Source/Runtime/AudioExtensions/AudioExtensions.Build.cs
+++ b/Engine/Source/Runtime/AudioExtensions/AudioExtensions.Build.cs
@@ -6,6 +6,8 @@ namespace UnrealBuildTool.Rules
 	{
 		public AudioExtensions(ReadOnlyTargetRules Target) : base(Target)
 		{
+			OptimizeCode = CodeOptimization.Never;
+
 			PublicDependencyModuleNames.AddRange(
 				new string[]
 				{
@@ -15,14 +17,6 @@ namespace UnrealBuildTool.Rules
 					"AudioMixerCore"
 				}
 			);
-
-            PrivateDependencyModuleNames.AddRange(
-                new string[]
-                {
-					// This is needed for SoundfieldStatics::GetDefaultPositionMap.
-                    //"AudioMixer"
-                }
-            );
-        }
+		}
 	}
 }
diff --git a/Engine/Source/Runtime/AudioExtensions/Private/IAudioModulation.cpp b/Engine/Source/Runtime/AudioExtensions/Private/IAudioModulation.cpp
new file mode 100644
--- /dev/null
+++ b/Engine/Source/Runtime/AudioExtensions/Private/IAudioModulation.cpp
@@ -0,0 +1,133 @@
+// Copyright Epic Games, Inc. All Rights Reserved.
+#include "IAudioModulation.h"
+
+#include "UObject/Object.h"
+
+
+namespace Audio
+{
+	FModulatorHandle::FModulatorHandle()
+		: ParentId(INDEX_NONE)
+		, ModulatorId(INDEX_NONE)
+		, Modulation(nullptr)
+	{
+	}
+
+	FModulatorHandle::FModulatorHandle(IAudioModulation& InModulation, uint32 InParentId, const USoundModulatorBase& InModulatorBase)
+	{
+		if (InModulation.RegisterModulator(InParentId, InModulatorBase))
+		{
+			ParentId = InParentId;
+			ModulatorId = static_cast<Audio::FModulatorId>(InModulatorBase.GetUniqueID());
+			Modulation = &InModulation;
+		}
+		else
+		{
+			ParentId = INDEX_NONE;
+			ModulatorId = INDEX_NONE;
+			Modulation = nullptr;
+		}
+	}
+
+	FModulatorHandle::FModulatorHandle(const FModulatorHandle& InOther)
+	{
+		if (InOther.Modulation && InOther.Modulation->RegisterModulator(InOther.ParentId, InOther.ModulatorId))
+		{
+			ParentId = InOther.ParentId;
+			ModulatorId = static_cast<Audio::FModulatorId>(InOther.ModulatorId);
+			Modulation = InOther.Modulation;
+		}
+		else
+		{
+			ParentId = INDEX_NONE;
+			ModulatorId = INDEX_NONE;
+			Modulation = nullptr;
+		}
+	}
+
+	FModulatorHandle::FModulatorHandle(FModulatorHandle&& InOther)
+		: ParentId(InOther.ParentId)
+		, ModulatorId(InOther.ModulatorId)
+		, Modulation(InOther.Modulation)
+	{
+		// Move does not register as presumed already activated or
+		// copying default handle, which is invalid. Removes data
+		// from handle being moved to avoid double deactivation on
+		// destruction.
+		InOther.ParentId = INDEX_NONE;
+		InOther.ModulatorId = INDEX_NONE;
+		InOther.Modulation = nullptr;
+	}
+
+	FModulatorHandle::~FModulatorHandle()
+	{
+		if (Modulation)
+		{
+			Modulation->UnregisterModulator(*this);
+		}
+	}
+
+	FModulatorHandle& FModulatorHandle::operator=(const FModulatorHandle& InOther)
+	{
+		if (InOther.Modulation && InOther.Modulation->RegisterModulator(InOther.ParentId, InOther.ModulatorId))
+		{
+			ParentId = InOther.ParentId;
+			ModulatorId = static_cast<Audio::FModulatorId>(InOther.ModulatorId);
+			Modulation = InOther.Modulation;
+		}
+		else
+		{
+			ParentId = INDEX_NONE;
+			ModulatorId = INDEX_NONE;
+			Modulation = nullptr;
+		}
+
+		return *this;
+	}
+
+	FModulatorHandle& FModulatorHandle::operator=(FModulatorHandle&& InOther)
+	{
+		// Move does not activate as presumed already activated or
+		// copying default handle, which is invalid. Removes data
+		// from handle being moved to avoid double deactivation on
+		// destruction.
+		ParentId = InOther.ParentId;
+		ModulatorId = InOther.ModulatorId;
+		Modulation = InOther.Modulation;
+
+		InOther.ParentId	= INDEX_NONE;
+		InOther.ModulatorId = INDEX_NONE;
+		InOther.Modulation	= nullptr;
+
+		return *this;
+	}
+
+	FModulatorId FModulatorHandle::GetId() const
+	{
+		return ModulatorId;
+	}
+
+	uint32 FModulatorHandle::GetParentId() const
+	{
+		return ParentId;
+	}
+
+	float FModulatorHandle::GetValue(const float InDefaultValue) const
+	{
+		if (IsValid())
+		{
+			float Value = InDefaultValue;
+			if (Modulation->GetModulatorValue(*this, Value))
+			{
+				return Value;
+			}
+		}
+
+		return InDefaultValue;
+	}
+
+	bool FModulatorHandle::IsValid() const
+	{
+		return Modulation && ModulatorId != INDEX_NONE;
+	}
+} // namespace Audio
diff --git a/Engine/Source/Runtime/AudioExtensions/Public/IAudioExtensionPlugin.h b/Engine/Source/Runtime/AudioExtensions/Public/IAudioExtensionPlugin.h
--- a/Engine/Source/Runtime/AudioExtensions/Public/IAudioExtensionPlugin.h
+++ b/Engine/Source/Runtime/AudioExtensions/Public/IAudioExtensionPlugin.h
@@ -12,15 +12,8 @@
 
 #include "IAudioExtensionPlugin.generated.h"
 
-// Forward Declarations
-
-#if !UE_BUILD_SHIPPING
-class FCanvas;
-class FCommonViewportClient;
-class FViewport;
-class UFont;
-#endif // !UE_BUILD_SHIPPING
 
+// Forward Declarations
 class FAudioDevice;
 class FSoundEffectBase;
 class FSoundEffectSource;
@@ -451,59 +444,6 @@ public:
 	}
 };
 
-
-/** Override to provide users with modulation settings custom to individual sounds */
-UCLASS(config = Engine, abstract, editinlinenew, BlueprintType)
-class AUDIOEXTENSIONS_API USoundModulationPluginSourceSettingsBase : public UObject
-{
-	GENERATED_BODY()
-};
-
-/** Collection of settings available on sound objects */
-USTRUCT(BlueprintType)
-struct AUDIOEXTENSIONS_API FSoundModulation
-{
-	GENERATED_USTRUCT_BODY()
-
-public:
-	UPROPERTY(EditAnywhere, BlueprintReadWrite, Category = Modulation)
-	TArray<USoundModulationPluginSourceSettingsBase*> Settings;
-};
-
-/** Interface to sound that is modulateable, allowing for certain specific 
-  * behaviors to be controlled on the sound level by the modulation system.
-  */
-class AUDIOEXTENSIONS_API ISoundModulatable
-{
-public:
-	virtual ~ISoundModulatable() = default;
-
-	/**
-	 * Returns the modulation settings of the sound
-	 */
-	virtual USoundModulationPluginSourceSettingsBase* FindModulationSettings() const = 0;
-
-	/**
-	 * Gets the object definition id of the given playing sound's instance
-	 */
-	virtual uint32 GetObjectId() const = 0;
-
-	/**
-	 * Returns number of actively instances of sound playing (including virtualized instances)
-	 */
-	virtual int32 GetPlayCount() const = 0;
-
-	/**
-	 * Returns whether or not sound is an editor preview sound
-	 */
-	virtual bool IsPreviewSound() const = 0;
-
-	/**
-	 * Stops sound.
-	 */
-	virtual void Stop() = 0;
-};
-
 /************************************************************************/
 /* IAudioModulationFactory                                              */
 /*                                                                      */
@@ -537,182 +477,6 @@ public:
 };
 
 
-/*
- * Modulatable controls found on each sound instance
- * processed by the modulation plugin enabled
- */
-struct AUDIOEXTENSIONS_API FSoundModulationControls
-{
-	float Volume;
-	float Pitch;
-	float Lowpass;
-	float Highpass;
-
-	TMap<FName, float> Controls;
-
-	FSoundModulationControls()
-		: Volume(1.0f)
-		, Pitch(1.0f)
-		, Lowpass(MAX_FILTER_FREQUENCY)
-		, Highpass(MIN_FILTER_FREQUENCY)
-	{
-	}
-};
-
-/** Parameter allowing modulation control override for systems opting in to the Modulation System. */
-USTRUCT(BlueprintType)
-struct  FSoundModulationParameter
-{
-	GENERATED_BODY()
-
-public:
-	/** Name of modulation control to drive parameter. Uses value last cached if control is or becomes invalid. */
-	UPROPERTY(EditAnywhere, BlueprintReadWrite, Category = Modulation)
-	FName Control;
-
-	/** Default modulation parameter value. */
-	UPROPERTY(EditAnywhere, BlueprintReadWrite, Category = Modulation)
-	float Value;
-
-private:
-	float ValueMin;
-	float ValueMax;
-
-public:
-	FSoundModulationParameter()
-		: Control(NAME_None)
-		, Value(0.0f)
-		, ValueMin(0.0f)
-		, ValueMax(1.0f)
-	{
-	}
-
-	FSoundModulationParameter(float InValue, float InValueMin = TNumericLimits<float>::Min(), float InValueMax = TNumericLimits<float>::Max())
-		: Control(NAME_None)
-		, ValueMin(InValueMin)
-		, ValueMax(InValueMax)
-	{
-		Value = FMath::Clamp(InValue, ValueMin, ValueMax);
-	}
-
-	FORCEINLINE operator float() const
-	{
-		return Value;
-	}
-
-	FORCEINLINE FSoundModulationParameter& operator= (const FSoundModulationParameter& InParam)
-	{
-		Control = InParam.Control;
-		ValueMin = InParam.ValueMin;
-		ValueMax = InParam.ValueMax;
-		Value = FMath::Clamp(InParam.Value, ValueMin, ValueMax);
-
-		return *this;
-	}
-
-	FORCEINLINE FSoundModulationParameter& operator= (float InValue)
-	{
-		Value = FMath::Clamp(InValue, ValueMin, ValueMax);
-		return *this;
-	}
-
-	FORCEINLINE float operator* (float InValue)
-	{
-		return InValue * Value;
-	}
-
-	FORCEINLINE float operator/ (float InValue)
-	{
-		return InValue / Value;
-	}
-
-	FORCEINLINE float GetValue() const
-	{
-		return Value;
-	}
-
-	FORCEINLINE float GetMinValue() const
-	{
-		return ValueMin;
-	}
-
-	FORCEINLINE float GetMaxValue() const
-	{
-		return ValueMax;
-	}
-
-	FORCEINLINE void SetBounds(float InValueMin, float InValueMax)
-	{
-		ValueMin = InValueMin;
-		ValueMax = InValueMax;
-
-		Value = FMath::Clamp(Value, InValueMin, InValueMax);
-	}
-
-	/* Returns true if updated by control, false if not */
-	bool SetValue(const FSoundModulationControls& InModControls)
-	{
-		if (Control != NAME_None)
-		{
-			if (const float* ControlValue = InModControls.Controls.Find(Control))
-			{
-				Value = FMath::Clamp(*ControlValue, ValueMin, ValueMax);
-				return true;
-			}
-		}
-
-		Value = FMath::Clamp(Value, ValueMin, ValueMax);
-		return false;
-	}
-};
-
-using FSoundModulationControlIndex = uint32;
-
-class IAudioModulation
-{
-public:
-	/** Virtual destructor */
-	virtual ~IAudioModulation() { }
-
-	/** Calculates initial volume to determine if sound is audible using base settings data */
-	virtual float CalculateInitialVolume(const USoundModulationPluginSourceSettingsBase& InSettingsBase) { return 1.0f; }
-
-	/** Initialize the modulation plugin with the same rate and number of sources */
-	virtual void Initialize(const FAudioPluginInitializationParams& InitializationParams) { }
-
-	/** Called when a USoundBase type begins playing a sound */
-	virtual void OnInitSound(ISoundModulatable& Sound, const USoundModulationPluginSourceSettingsBase& Settings) { }
-
-	/** Called when a source is assigned to a voice */
-	virtual void OnInitSource(const uint32 SourceId, const FName& AudioComponentUserId, const uint32 NumChannels, const USoundModulationPluginSourceSettingsBase& Settings) { }
-
-	/** Called when a source is done playing and is released */
-	virtual void OnReleaseSource(const uint32 SourceId) { }
-
-	/** Called when a USoundBase type stops playing any sounds */
-	virtual void OnReleaseSound(ISoundModulatable& Sound) { }
-
-#if !UE_BUILD_SHIPPING
-	/** Request to post help from active plugin (non-shipping builds only) */
-	virtual bool OnPostHelp(FCommonViewportClient* ViewportClient, const TCHAR* Stream) { return false; };
-
-	/** Render stats pertaining to modulation (non-shipping builds only) */
-	virtual int32 OnRenderStat(FViewport* Viewport, FCanvas* Canvas, int32 X, int32 Y, const UFont& Font, const FVector* ViewLocation, const FRotator* ViewRotation) { return Y; }
-
-	/** Toggle showing render stats pertaining to modulation (non-shipping builds only) */
-	virtual bool OnToggleStat(FCommonViewportClient* ViewportClient, const TCHAR* Stream) { return false; }
-#endif //!UE_BUILD_SHIPPING
-
-	/** Processes audio with the given input and output data structs.*/
-	virtual void ProcessAudio(const FAudioPluginSourceInputData& InputData, FAudioPluginSourceOutputData& OutputData) { }
-
-	/** Processes modulated sound controls, returning whether or not controls were modified and an update is pending. */
-	virtual bool ProcessControls(const uint32 SourceId, FSoundModulationControls& Controls) { return false; }
-
-	/** Processes all modulators */
-	virtual void ProcessModulators(const float Elapsed) { }
-};
-
 /** This is a class which should be overridden to provide users with settings to use for individual sounds */
 UCLASS(config = Engine, abstract, editinlinenew, BlueprintType)
 class AUDIOEXTENSIONS_API UReverbPluginSourceSettingsBase : public UObject
diff --git a/Engine/Source/Runtime/AudioExtensions/Public/IAudioModulation.h b/Engine/Source/Runtime/AudioExtensions/Public/IAudioModulation.h
new file mode 100644
--- /dev/null
+++ b/Engine/Source/Runtime/AudioExtensions/Public/IAudioModulation.h
@@ -0,0 +1,321 @@
+// Copyright Epic Games, Inc. All Rights Reserved.
+#pragma once
+
+#include "AudioDefines.h"
+#include "CoreMinimal.h"
+#include "DSP/BufferVectorOperations.h"
+#include "IAudioExtensionPlugin.h"
+
+#include "IAudioModulation.generated.h"
+
+
+// Forward Declarations
+class IAudioModulation;
+class ISoundModulatable;
+class USoundModulatorBase;
+class UObject;
+
+#if !UE_BUILD_SHIPPING
+class FCanvas;
+class FCommonViewportClient;
+class FViewport;
+class UFont;
+#endif // !UE_BUILD_SHIPPING
+
+struct FSoundModulationControls;
+
+namespace Audio
+{
+	using FModulatorId = uint32;
+	using FControlModulatorValueMap = TMap<FModulatorId, float>;
+
+	struct AUDIOEXTENSIONS_API FModulatorHandle
+	{
+		FModulatorHandle();
+		FModulatorHandle(IAudioModulation& InModulation, uint32 InParentId, const USoundModulatorBase& InModulatorBase);
+		FModulatorHandle(const FModulatorHandle& InOther);
+		FModulatorHandle(FModulatorHandle&& InOther);
+
+		~FModulatorHandle();
+
+		FModulatorHandle& operator=(const FModulatorHandle& InOther);
+		FModulatorHandle& operator=(FModulatorHandle&& InOther);
+
+		FModulatorId GetId() const;
+		uint32 GetParentId() const;
+		float GetValue(const float InDefaultValue = 1.0f) const;
+		bool IsValid() const;
+
+	private:
+		uint32 ParentId;
+		FModulatorId ModulatorId;
+		IAudioModulation* Modulation;
+	};
+} // namespace Audio
+
+
+UENUM()
+enum class ESoundModulatorOperator : uint8
+{
+	/** Modulation is disabled */
+	None = 0 UMETA(DisplayName = "None"),
+
+	/** Modulator result is multiplier of input value */
+	Multiply,
+
+	/** Modulator result is divisor, input value is dividend */
+	Divide,
+
+	/** Take the minimum of the modulator result and input value */
+	Min,
+
+	/** Take the maximum of the modulator result and input value */
+	Max,
+
+	/** Add modulator result and input value */
+	Add,
+
+	/** Subtract modulator result from input value */
+	Subtract,
+
+	Count UMETA(Hidden),
+};
+
+
+namespace SoundModulatorOperator
+{
+	FORCEINLINE float GetDefaultValue(ESoundModulatorOperator InOperator, float InMin, float InMax)
+	{
+		switch (InOperator)
+		{
+			case ESoundModulatorOperator::Max:
+			{
+				return InMin;
+			}
+			break;
+
+			case ESoundModulatorOperator::Min:
+			{
+				return InMax;
+			}
+			break;
+
+			case ESoundModulatorOperator::Multiply:
+			case ESoundModulatorOperator::Divide:
+			{
+				return 1.0f;
+			}
+			break;
+
+			case ESoundModulatorOperator::Add:
+			case ESoundModulatorOperator::Subtract:
+			case ESoundModulatorOperator::None:
+			{
+				return 0.0f;
+			}
+			break;
+
+			default:
+			{
+				static_assert(static_cast<int32>(ESoundModulatorOperator::Count) == 7, "Possible missing operator switch case coverage");
+				return NAN;
+			}
+			break;
+		}
+	}
+
+	FORCEINLINE float Apply(ESoundModulatorOperator InOperator, float InValueA, float InValueB)
+	{
+		switch (InOperator)
+		{
+			case ESoundModulatorOperator::Max:
+			{
+				return FMath::Max(InValueA, InValueB);
+			}
+			break;
+
+			case ESoundModulatorOperator::Min:
+			{
+				return FMath::Min(InValueA, InValueB);
+			}
+			break;
+
+			case ESoundModulatorOperator::Multiply:
+			{
+				return InValueA * InValueB;
+			}
+			break;
+
+			case ESoundModulatorOperator::Divide:
+			{
+				return InValueA / InValueB;
+			}
+			break;
+
+			case ESoundModulatorOperator::Add:
+			{
+				return InValueA + InValueB;
+			}
+			break;
+
+			case ESoundModulatorOperator::Subtract:
+			{
+				return InValueA - InValueB;
+			}
+			break;
+
+			case ESoundModulatorOperator::None:
+			{
+				return InValueA;
+			}
+			break;
+
+			default:
+			{
+				checkf(false, TEXT("Cannot apply 'None' as operator to modulator"));
+				static_assert(static_cast<int32>(ESoundModulatorOperator::Count) == 7, "Possible missing operator switch case coverage");
+				return NAN;
+			}
+			break;
+		}
+	}
+} // namespace SoundModulatorOperator
+
+/*
+ * Modulateable controls found on each sound instance
+ * processed by the enabled modulation plugin.
+ */
+struct AUDIOEXTENSIONS_API FSoundModulationControls
+{
+	float Volume;
+	float Pitch;
+	float Lowpass;
+	float Highpass;
+
+	FSoundModulationControls()
+		: Volume(1.0f)
+		, Pitch(1.0f)
+		, Lowpass(MAX_FILTER_FREQUENCY)
+		, Highpass(MIN_FILTER_FREQUENCY)
+	{
+	}
+};
+
+class AUDIOEXTENSIONS_API IAudioModulation
+{
+public:
+	/** Virtual destructor */
+	virtual ~IAudioModulation() { }
+
+	/** Calculates initial volume to determine if sound is audible using base settings data */
+	virtual float CalculateInitialVolume(const USoundModulationPluginSourceSettingsBase& InSettingsBase) { return 1.0f; }
+
+	/** Initialize the modulation plugin with the same rate and number of sources */
+	virtual void Initialize(const FAudioPluginInitializationParams& InitializationParams) { }
+
+	/** Called when a USoundBase type begins playing a sound */
+	virtual void OnInitSound(ISoundModulatable& Sound, const USoundModulationPluginSourceSettingsBase& Settings) { }
+
+	/** Called when a source is assigned to a voice */
+	virtual void OnInitSource(const uint32 SourceId, const FName& AudioComponentUserId, const uint32 NumChannels, const USoundModulationPluginSourceSettingsBase& Settings) { }
+
+	/** Called when a source is done playing and is released */
+	virtual void OnReleaseSource(const uint32 SourceId) { }
+
+	/** Called when a USoundBase type stops playing any sounds */
+	virtual void OnReleaseSound(ISoundModulatable& Sound) { }
+
+#if !UE_BUILD_SHIPPING
+	/** Request to post help from active plugin (non-shipping builds only) */
+	virtual bool OnPostHelp(FCommonViewportClient* ViewportClient, const TCHAR* Stream) { return false; };
+
+	/** Render stats pertaining to modulation (non-shipping builds only) */
+	virtual int32 OnRenderStat(FViewport* Viewport, FCanvas* Canvas, int32 X, int32 Y, const UFont& Font, const FVector* ViewLocation, const FRotator* ViewRotation) { return Y; }
+
+	/** Toggle showing render stats pertaining to modulation (non-shipping builds only) */
+	virtual bool OnToggleStat(FCommonViewportClient* ViewportClient, const TCHAR* Stream) { return false; }
+#endif //!UE_BUILD_SHIPPING
+
+	/** Run on the audio render thread prior to processing audio */
+	virtual void OnBeginAudioRenderThreadUpdate() { }
+
+	/** Processes audio with the given input and output data structs.*/
+	virtual void ProcessAudio(const FAudioPluginSourceInputData& InputData, FAudioPluginSourceOutputData& OutputData) { }
+
+	/** Processes modulated sound controls, returning whether or not controls were modified and an update is pending. */
+	virtual bool ProcessControls(const uint32 SourceId, FSoundModulationControls& Controls) { return false; }
+
+	/** Processes all modulators */
+	virtual void ProcessModulators(const float Elapsed) { }
+
+	virtual void UpdateModulator(const USoundModulatorBase& InModulator) { }
+
+protected:
+	virtual bool RegisterModulator(uint32 InParentId, const USoundModulatorBase& InModulatorBase) { return false; }
+	virtual bool RegisterModulator(uint32 InParentId, Audio::FModulatorId InModulatorId) { return false; }
+	virtual bool GetModulatorValue(const Audio::FModulatorHandle& ModulatorHandle, float& OutValue) { return false; }
+	virtual void UnregisterModulator(const Audio::FModulatorHandle& InHandle) { }
+
+	friend Audio::FModulatorHandle;
+};
+
+/**
+ * Base class for all modulators
+ */
+UCLASS(config = Engine, abstract, editinlinenew, BlueprintType)
+class AUDIOEXTENSIONS_API USoundModulatorBase : public UObject
+{
+	GENERATED_BODY()
+};
+
+/** Override to provide users with modulation settings custom to individual sounds */
+UCLASS(config = Engine, abstract, editinlinenew, BlueprintType)
+class AUDIOEXTENSIONS_API USoundModulationPluginSourceSettingsBase : public UObject
+{
+	GENERATED_BODY()
+};
+
+/** Collection of settings available on sound objects */
+USTRUCT(BlueprintType)
+struct AUDIOEXTENSIONS_API FSoundModulation
+{
+	GENERATED_USTRUCT_BODY()
+
+public:
+	UPROPERTY(EditAnywhere, BlueprintReadWrite, Category = Modulation)
+		TArray<USoundModulationPluginSourceSettingsBase*> Settings;
+};
+
+/** Interface to sound that is modulateable, allowing for certain specific
+  * behaviors to be controlled on the sound level by the modulation system.
+  */
+class AUDIOEXTENSIONS_API ISoundModulatable
+{
+public:
+	virtual ~ISoundModulatable() = default;
+
+	/**
+	 * Returns the modulation settings of the sound
+	 */
+	virtual USoundModulationPluginSourceSettingsBase* FindModulationSettings() const = 0;
+
+	/**
+	 * Gets the object definition id of the given playing sound's instance
+	 */
+	virtual uint32 GetObjectId() const = 0;
+
+	/**
+	 * Returns number of actively instances of sound playing (including virtualized instances)
+	 */
+	virtual int32 GetPlayCount() const = 0;
+
+	/**
+	 * Returns whether or not sound is an editor preview sound
+	 */
+	virtual bool IsPreviewSound() const = 0;
+
+	/**
+	 * Stops sound.
+	 */
+	virtual void Stop() = 0;
+};
diff --git a/Engine/Source/Runtime/AudioMixer/Private/AudioMixerSourceManager.cpp b/Engine/Source/Runtime/AudioMixer/Private/AudioMixerSourceManager.cpp
--- a/Engine/Source/Runtime/AudioMixer/Private/AudioMixerSourceManager.cpp
+++ b/Engine/Source/Runtime/AudioMixer/Private/AudioMixerSourceManager.cpp
@@ -2648,6 +2648,12 @@ namespace Audio
 			PumpCommandQueue();
 		}
 
+		// Notify modulation interface that we are beginning to update
+		if (MixerDevice->IsModulationPluginEnabled() && MixerDevice->ModulationInterface.IsValid())
+		{
+			MixerDevice->ModulationInterface->OnBeginAudioRenderThreadUpdate();
+		}
+
 		// Update pending tasks and release them if they're finished
 		UpdatePendingReleaseData();
 
diff --git a/Engine/Source/Runtime/Engine/Classes/Sound/SoundEffectSource.h b/Engine/Source/Runtime/Engine/Classes/Sound/SoundEffectSource.h
--- a/Engine/Source/Runtime/Engine/Classes/Sound/SoundEffectSource.h
+++ b/Engine/Source/Runtime/Engine/Classes/Sound/SoundEffectSource.h
@@ -3,6 +3,7 @@
 
 #include "CoreMinimal.h"
 #include "IAudioExtensionPlugin.h"
+#include "IAudioModulation.h"
 #include "Sound/SoundEffectPreset.h"
 #include "Sound/SoundEffectBase.h"
 #include "Templates/SharedPointer.h"
diff --git a/Engine/Source/Runtime/Engine/Classes/Sound/SoundModulationParameter.h b/Engine/Source/Runtime/Engine/Classes/Sound/SoundModulationParameter.h
new file mode 100644
--- /dev/null
+++ b/Engine/Source/Runtime/Engine/Classes/Sound/SoundModulationParameter.h
@@ -0,0 +1,96 @@
+// Copyright Epic Games, Inc. All Rights Reserved.
+#pragma once
+
+#include "CoreMinimal.h"
+#include "DSP/BufferVectorOperations.h"
+#include "IAudioModulation.h"
+
+#include "SoundModulationParameter.generated.h"
+
+
+// Forward Declarations
+class USoundModulatorBase;
+class UObject;
+
+
+/** Parameter settings allowing modulation control override for systems opting in to the Modulation System. */
+USTRUCT(BlueprintType)
+struct ENGINE_API FSoundModulationParameterSettings
+{
+	GENERATED_USTRUCT_BODY()
+
+	FSoundModulationParameterSettings();
+
+	/** Base value of parameter */
+	UPROPERTY(EditAnywhere, BlueprintReadWrite, Category = Modulation)
+	float Value;
+
+	/** Operator to apply when modulating the default value. */
+	UPROPERTY(EditAnywhere, BlueprintReadWrite, Category = Modulation)
+	ESoundModulatorOperator Operator;
+
+	/** Subscribed modulator to listen to apply result to base value. */
+	UPROPERTY(EditAnywhere, BlueprintReadWrite, Category = Modulation, meta = (EditCondition = "Operator != ESoundModulatorOperator::None", EditConditionHides))
+	USoundModulatorBase* Modulator;
+};
+
+namespace Audio
+{
+	struct ENGINE_API FModulationParameter
+	{
+		public:
+			FModulationParameter() = default;
+
+			void Init(FDeviceId InDeviceId, uint32 InParentId, bool bInIsBuffered = false, float InValueMin = 0.0f, float InValueMax = 1.0f);
+
+			/* Updates internal buffer to modulated result with control value as provided by modulation plugin.
+			 * Returns true if control value was updated.
+			 */
+			bool ProcessControl(float* RESTRICT InBuffer, int32 InNumSamples);
+
+			/* Updates internal target (or buffer if set to bIsBuffered) to new parameter control value as provided by modulation plugin.
+			 * Returns true if value was updated.
+			 */
+			bool ProcessControl(float InValueBase, int32 InNumSamples = 0);
+
+			void SetMax(float InMax);
+			void SetMin(float InMin);
+			void UpdateSettings(const FSoundModulationParameterSettings& InSettings);
+
+		private:
+			FDeviceId DeviceId = INDEX_NONE;
+			uint32 ParentId = INDEX_NONE;
+
+			float ValueMin = 0.0f;
+			float ValueMax = 1.0f;
+			float ValueTarget = 0.0f;
+			float DefaultValue = 0.0f;
+
+			ESoundModulatorOperator Operator = ESoundModulatorOperator::None;
+
+			bool bIsBuffered = false;
+
+			AlignedFloatBuffer Buffer;
+			FModulatorHandle Handle;
+
+			FCriticalSection SettingsCritSection;
+
+		public:
+			FORCEINLINE const AlignedFloatBuffer& GetBuffer() const
+			{
+				check(bIsBuffered);
+				return Buffer;
+			}
+
+			FORCEINLINE float GetSample(int32 InSampleIndex) const
+			{
+				check(bIsBuffered);
+				return InSampleIndex < Buffer.Num() ? Buffer[InSampleIndex] : ValueTarget;
+			}
+
+			FORCEINLINE float GetTarget() const
+			{
+				return ValueTarget;
+			}
+	};
+} // namespace Audio
\ No newline at end of file
diff --git a/Engine/Source/Runtime/Engine/Private/SoundModulationParameter.cpp b/Engine/Source/Runtime/Engine/Private/SoundModulationParameter.cpp
new file mode 100644
--- /dev/null
+++ b/Engine/Source/Runtime/Engine/Private/SoundModulationParameter.cpp
@@ -0,0 +1,312 @@
+// Copyright Epic Games, Inc. All Rights Reserved.
+
+#include "Sound/SoundModulationParameter.h"
+
+#include "Async/Async.h"
+#include "AudioDevice.h"
+#include "Math/TransformCalculus.h"
+#include "UObject/Object.h"
+
+
+FSoundModulationParameterSettings::FSoundModulationParameterSettings()
+	: Value(0.0f)
+	, Operator(ESoundModulatorOperator::None)
+	, Modulator(nullptr)
+{
+}
+
+namespace Audio
+{
+	void FModulationParameter::Init(FDeviceId InDeviceId, uint32 InParentId, bool bInIsBuffered, float InValueMin, float InValueMax)
+	{
+		DeviceId = InDeviceId;
+		ParentId = InParentId;
+		bIsBuffered = bInIsBuffered;
+		ValueMin = InValueMin;
+		ValueMax = InValueMax;
+	}
+
+	bool FModulationParameter::ProcessControl(float* RESTRICT InBuffer, int32 InNumSamples)
+	{
+		check(bIsBuffered);
+
+		float LastTarget = ValueTarget;
+		ESoundModulatorOperator CurrentOperator = ESoundModulatorOperator::None;
+		bool bIsActive = false;
+
+		FScopeLock Lock(&SettingsCritSection);
+		{
+			bIsActive = Handle.IsValid();
+			ValueTarget = Handle.GetValue(DefaultValue);
+			CurrentOperator = Operator;
+		}
+
+		if (Buffer.Num() != InNumSamples)
+		{
+			Buffer.Reset();
+			Buffer.AddUninitialized(InNumSamples);
+		}
+		FMemory::Memcpy(Buffer.GetData(), InBuffer, sizeof(float) * InNumSamples);
+
+		const bool bTargetChanged = !FMath::IsNearlyEqual(LastTarget, ValueTarget);
+		if (!bIsActive || CurrentOperator == ESoundModulatorOperator::None)
+		{
+			return bTargetChanged;
+		}
+
+		switch (CurrentOperator)
+		{
+			case ESoundModulatorOperator::Max:
+			{
+				const float ValueDelta = (ValueTarget - LastTarget) / InNumSamples;
+				for (int32 Sample = 0; Sample < InNumSamples; ++Sample)
+				{
+					Buffer[Sample] = FMath::Max(Buffer[Sample], LastTarget + (Sample * ValueDelta));
+				}
+			}
+			break;
+
+			case ESoundModulatorOperator::Min:
+			{
+				const float ValueDelta = (ValueTarget - LastTarget) / InNumSamples;
+				for (int32 Sample = 0; Sample < InNumSamples; ++Sample)
+				{
+					Buffer[Sample] = FMath::Min(Buffer[Sample], LastTarget + (Sample * ValueDelta));
+				}
+			}
+			break;
+
+			case ESoundModulatorOperator::Multiply:
+			{
+				FadeBufferFast(Buffer, LastTarget, ValueTarget);
+			}
+			break;
+
+			case ESoundModulatorOperator::Divide:
+			{
+				// Modulators could hypothetically be negative, but should avoid divide-by-zero
+				// Buffer clamp later constricts appropriately for use case
+				if (LastTarget == 0.0f)
+				{
+					LastTarget = SMALL_NUMBER;
+				}
+
+				if (ValueTarget == 0.0f)
+				{
+					ValueTarget = SMALL_NUMBER;
+				}
+
+				FadeBufferFast(Buffer, Inverse(LastTarget), Inverse(ValueTarget));
+			}
+			break;
+
+			case ESoundModulatorOperator::Add:
+			{
+				const float ValueDelta = (ValueTarget - LastTarget) / InNumSamples;
+				AddConstantToBufferInplace(Buffer, LastTarget);
+				for (int32 Sample = 0; Sample < InNumSamples; ++Sample)
+				{
+					Buffer[Sample] += Sample * ValueDelta;
+				}
+			}
+			break;
+
+			case ESoundModulatorOperator::Subtract:
+			{
+				const float ValueDelta = (ValueTarget - LastTarget) / InNumSamples;
+				AddConstantToBufferInplace(Buffer, -LastTarget);
+				for (int32 Sample = 0; Sample < InNumSamples; ++Sample)
+				{
+					Buffer[Sample] -= Sample * ValueDelta;
+				}
+			}
+			break;
+
+			case ESoundModulatorOperator::None:
+			default:
+			{
+				static_assert(static_cast<int32>(ESoundModulatorOperator::Count) == 7, "Possible missing operator switch case coverage");
+			}
+			break;
+		}
+
+		BufferRangeClampFast(Buffer, ValueMin, ValueMax);
+		return bTargetChanged;
+	}
+
+	bool FModulationParameter::ProcessControl(float InValueBase, int32 InNumSamples)
+	{
+		float LastTarget = ValueTarget;
+
+		ESoundModulatorOperator CurrentOperator = ESoundModulatorOperator::None;
+		bool bIsActive = false;
+
+		FScopeLock Lock(&SettingsCritSection);
+		{
+			bIsActive = Handle.IsValid();
+			ValueTarget = Handle.GetValue(DefaultValue);
+			CurrentOperator = Operator;
+		}
+
+		if (!bIsActive || Operator == ESoundModulatorOperator::None)
+		{
+			Buffer.Reset();
+			ValueTarget = InValueBase;
+			return !FMath::IsNearlyEqual(LastTarget, ValueTarget);
+		}
+
+		if (!bIsBuffered)
+		{
+			ValueTarget = SoundModulatorOperator::Apply(CurrentOperator, InValueBase, ValueTarget);
+			ValueTarget = FMath::Clamp(ValueTarget, ValueMin, ValueMax);
+			return !FMath::IsNearlyEqual(LastTarget, ValueTarget);
+		}
+
+		if (Buffer.Num() != InNumSamples)
+		{
+			Buffer.Reset();
+			Buffer.AddUninitialized(InNumSamples);
+		}
+
+		const bool bTargetChanged = !FMath::IsNearlyEqual(LastTarget, ValueTarget);
+		if (!bTargetChanged)
+		{
+			const float ModValue = SoundModulatorOperator::Apply(CurrentOperator, InValueBase, ValueTarget);
+			ValueTarget = FMath::Clamp(ModValue, ValueMin, ValueMax);
+
+			BufferSetToConstantInplace(Buffer, ValueTarget);
+			return false;
+		}
+
+		BufferSetToConstantInplace(Buffer, InValueBase);
+
+		switch (CurrentOperator)
+		{
+			case ESoundModulatorOperator::Max:
+			{
+				const float ValueDelta = (ValueTarget - LastTarget) / InNumSamples;
+				for (int32 Sample = 0; Sample < InNumSamples; ++Sample)
+				{
+					Buffer[Sample] = FMath::Max(Buffer[Sample], LastTarget + (Sample * ValueDelta));
+				}
+			}
+			break;
+
+			case ESoundModulatorOperator::Min:
+			{
+				const float ValueDelta = (ValueTarget - LastTarget) / InNumSamples;
+				for (int32 Sample = 0; Sample < InNumSamples; ++Sample)
+				{
+					Buffer[Sample] = FMath::Min(Buffer[Sample], LastTarget + (Sample * ValueDelta));
+				}
+			}
+			break;
+
+			case ESoundModulatorOperator::Multiply:
+			{
+				FadeBufferFast(Buffer, LastTarget, ValueTarget);
+			}
+			break;
+
+			case ESoundModulatorOperator::Divide:
+			{
+				// Modulators could hypothetically be negative, but should avoid divide-by-zero
+				// Buffer clamp later constricts appropriately for use case
+				if (LastTarget == 0.0f)
+				{
+					LastTarget = SMALL_NUMBER;
+				}
+
+				if (ValueTarget == 0.0f)
+				{
+					ValueTarget = SMALL_NUMBER;
+				}
+
+				FadeBufferFast(Buffer, Inverse(LastTarget), Inverse(ValueTarget));
+			}
+			break;
+
+			case ESoundModulatorOperator::Add:
+			{
+				const float ValueDelta = (ValueTarget - LastTarget) / InNumSamples;
+				AddConstantToBufferInplace(Buffer, LastTarget);
+				for (int32 Sample = 0; Sample < InNumSamples; ++Sample)
+				{
+					Buffer[Sample] += Sample * ValueDelta;
+				}
+			}
+			break;
+
+			case ESoundModulatorOperator::Subtract:
+			{
+				const float ValueDelta = (ValueTarget - LastTarget) / InNumSamples;
+				AddConstantToBufferInplace(Buffer, -LastTarget);
+				for (int32 Sample = 0; Sample < InNumSamples; ++Sample)
+				{
+					Buffer[Sample] -= Sample * ValueDelta;
+				}
+			}
+			break;
+
+			case ESoundModulatorOperator::None:
+			default:
+			{
+				static_assert(static_cast<int32>(ESoundModulatorOperator::Count) == 7, "Possible missing operator switch case coverage");
+			}
+			break;
+		}
+
+		BufferRangeClampFast(Buffer, ValueMin, ValueMax);
+		return bTargetChanged;
+	}
+
+	void FModulationParameter::SetMax(float InMax)
+	{
+		ValueMax = InMax;
+	}
+
+	void FModulationParameter::SetMin(float InMin)
+	{
+		ValueMin = InMin;
+	}
+
+	void FModulationParameter::UpdateSettings(const FSoundModulationParameterSettings& InSettings)
+	{
+		const ESoundModulatorOperator NewOperator = InSettings.Operator;
+		const TWeakObjectPtr<const USoundModulatorBase> ModPtr(InSettings.Modulator);
+		auto UpdateHandleLambda = [this, ModPtr, NewOperator]()
+		{
+			const float NewDefaultValue = SoundModulatorOperator::GetDefaultValue(NewOperator, ValueMin, ValueMax);
+
+			if (NewOperator != ESoundModulatorOperator::None && ModPtr.IsValid())
+			{
+				if (FAudioDevice* AudioDevice = FAudioDeviceManager::Get()->GetAudioDeviceRaw(DeviceId))
+				{
+					if (AudioDevice->IsModulationPluginEnabled() && AudioDevice->ModulationInterface.IsValid())
+					{
+						IAudioModulation* Modulation = AudioDevice->ModulationInterface.Get();
+
+						FScopeLock Lock(&SettingsCritSection);
+						{
+							DefaultValue = NewDefaultValue;
+							Operator = NewOperator;
+							Handle = FModulatorHandle(*Modulation, ParentId, *ModPtr.Get());
+						}
+						return;
+					}
+				}
+			}
+
+			FScopeLock Lock(&SettingsCritSection);
+			{
+				DefaultValue = NewDefaultValue;
+				Operator = NewOperator;
+				Handle = FModulatorHandle();
+			}
+		};
+
+		IsInAudioThread()
+			? UpdateHandleLambda()
+			: AsyncTask(ENamedThreads::AudioThread, MoveTemp(UpdateHandleLambda));
+	}
+} // namespace Audio
diff --git a/Engine/Source/Runtime/Engine/Public/Audio.h b/Engine/Source/Runtime/Engine/Public/Audio.h
--- a/Engine/Source/Runtime/Engine/Public/Audio.h
+++ b/Engine/Source/Runtime/Engine/Public/Audio.h
@@ -16,6 +16,7 @@
 #include "Sound/SoundSubmixSend.h"
 #include "Sound/SoundSourceBusSend.h"
 #include "IAudioExtensionPlugin.h"
+#include "IAudioModulation.h"
 
 ENGINE_API DECLARE_LOG_CATEGORY_EXTERN(LogAudio, Display, All);
 
