Fix circular dependency
#rnx

diff --git a/Engine/Plugins/Runtime/AudioModulation/Source/AudioModulation/Private/SoundControlBusMixProxy.h b/Engine/Plugins/Runtime/AudioModulation/Source/AudioModulation/Private/SoundControlBusMixProxy.h
--- a/Engine/Plugins/Runtime/AudioModulation/Source/AudioModulation/Private/SoundControlBusMixProxy.h
+++ b/Engine/Plugins/Runtime/AudioModulation/Source/AudioModulation/Private/SoundControlBusMixProxy.h
@@ -2,7 +2,7 @@
 #pragma once
 
 #include "IAudioModulation.h"
-#include "SoundModulationPatchProxy.h"
+#include "SoundControlBusProxy.h"
 #include "SoundModulationProxy.h"
 #include "SoundModulationValue.h"
 #include "SoundModulatorLFOProxy.h"
diff --git a/Engine/Plugins/Runtime/AudioModulation/Source/AudioModulation/Private/SoundModulationPatchProxy.h b/Engine/Plugins/Runtime/AudioModulation/Source/AudioModulation/Private/SoundModulationPatchProxy.h
--- a/Engine/Plugins/Runtime/AudioModulation/Source/AudioModulation/Private/SoundModulationPatchProxy.h
+++ b/Engine/Plugins/Runtime/AudioModulation/Source/AudioModulation/Private/SoundModulationPatchProxy.h
@@ -2,7 +2,6 @@
 #pragma once
 
 #include "SoundControlBusProxy.h"
-#include "SoundControlBusMixProxy.h"
 #include "SoundModulationProxy.h"
 #include "SoundModulatorLFOProxy.h"
 
