Swap Emplace_GetRef... no ref needed
#rnx

diff --git a/Engine/Plugins/Runtime/AudioModulation/Source/AudioModulation/Private/AudioModulationStatics.cpp b/Engine/Plugins/Runtime/AudioModulation/Source/AudioModulation/Private/AudioModulationStatics.cpp
--- a/Engine/Plugins/Runtime/AudioModulation/Source/AudioModulation/Private/AudioModulationStatics.cpp
+++ b/Engine/Plugins/Runtime/AudioModulation/Source/AudioModulation/Private/AudioModulationStatics.cpp
@@ -282,7 +282,7 @@ USoundControlBusMix* UAudioModulationStatics::CreateBusMix(const UObject* WorldC
 	{
 		if (Channel.Bus)
 		{
-			NewBusMix->Channels.Emplace_GetRef(Channel);
+			NewBusMix->Channels.Emplace(Channel);
 		}
 		else
 		{
