Fix modulation parameter assets not being loaded in non-editor builds

diff --git a/Engine/Plugins/Runtime/AudioModulation/Config/BaseAudioModulation.ini b/Engine/Plugins/Runtime/AudioModulation/Config/BaseAudioModulation.ini
new file mode 100644
--- /dev/null
+++ b/Engine/Plugins/Runtime/AudioModulation/Config/BaseAudioModulation.ini
@@ -0,0 +1,10 @@
+[/Script/AudioModulation.AudioModulationSettings]
++Parameters=/AudioModulation/BitDepth.BitDepth
++Parameters=/AudioModulation/HPFCutoffFrequency.HPFCutoffFrequency
++Parameters=/AudioModulation/LowRateFrequency.LowRateFrequency
++Parameters=/AudioModulation/LPFCutoffFrequency.LPFCutoffFrequency
++Parameters=/AudioModulation/Pan.Pan
++Parameters=/AudioModulation/Pitch.Pitch
++Parameters=/AudioModulation/SampleRate.SampleRate
++Parameters=/AudioModulation/TimeOfDay.TimeOfDay
++Parameters=/AudioModulation/Volume.Volume
\ No newline at end of file
diff --git a/Engine/Plugins/Runtime/AudioModulation/Source/AudioModulation/Private/AudioModulationSystem.cpp b/Engine/Plugins/Runtime/AudioModulation/Source/AudioModulation/Private/AudioModulationSystem.cpp
--- a/Engine/Plugins/Runtime/AudioModulation/Source/AudioModulation/Private/AudioModulationSystem.cpp
+++ b/Engine/Plugins/Runtime/AudioModulation/Source/AudioModulation/Private/AudioModulationSystem.cpp
@@ -6,6 +6,7 @@
 #include "Audio/AudioAddressPattern.h"
 #include "AudioModulationLogging.h"
 #include "AudioModulationProfileSerializer.h"
+#include "AudioModulationSettings.h"
 #include "AudioThread.h"
 #include "CoreGlobals.h"
 #include "Engine/Engine.h"
@@ -60,6 +61,23 @@ namespace AudioModulation
 
 	void FAudioModulationSystem::Initialize(const FAudioPluginInitializationParams& InitializationParams)
 	{
+		// Load all parameters listed in settings and keep loaded.  Must be done on initialization as
+		// soft object paths cannot be loaded on non-game threads (data from parameters is copied to a proxy
+		// in 'GetParameter' via the audio thread)
+		if (const UAudioModulationSettings* ModulationSettings = GetDefault<UAudioModulationSettings>())
+		{
+			for (const FSoftObjectPath& Path : ModulationSettings->Parameters)
+			{
+				if (UObject* ParamObj = Path.TryLoad())
+				{
+					ParamObj->AddToRoot();
+				}
+				else
+				{
+					UE_LOG(LogAudioModulation, Error, TEXT("Failed to load modulation parameter from path '%s'"), *Path.ToString());
+				}
+			}
+		}
 	}
 
 #if !UE_BUILD_SHIPPING
@@ -255,26 +273,29 @@ namespace AudioModulation
 			return Parameter;
 		}
 
-		for (TObjectIterator<USoundModulationParameter> Itr; Itr; ++Itr)
+		if (const UAudioModulationSettings* ModulationSettings = GetDefault<UAudioModulationSettings>())
 		{
-			if (USoundModulationParameter* Param = *Itr)
+			for (const FSoftObjectPath& Path : ModulationSettings->Parameters)
 			{
-				if (Param->GetFName() == InParamName)
+				if (const USoundModulationParameter* Param = Cast<const USoundModulationParameter>(Path.ResolveObject()))
 				{
-					Parameter.ParameterName			= InParamName;
-					Parameter.bRequiresConversion	= Param->RequiresUnitConversion();
-					Parameter.MixFunction			= Param->GetMixFunction();
-					Parameter.UnitFunction			= Param->GetUnitConversionFunction();
-					Parameter.LinearFunction		= Param->GetLinearConversionFunction();
-					Parameter.DefaultValue			= Param->GetUnitDefault();
-					Parameter.MinValue				= Param->GetUnitMin();
-					Parameter.MaxValue				= Param->GetUnitMax();
-					return Parameter;
+					if (Param->GetFName() == InParamName)
+					{
+						Parameter.ParameterName = InParamName;
+						Parameter.bRequiresConversion = Param->RequiresUnitConversion();
+						Parameter.MixFunction = Param->GetMixFunction();
+						Parameter.UnitFunction = Param->GetUnitConversionFunction();
+						Parameter.LinearFunction = Param->GetLinearConversionFunction();
+						Parameter.DefaultValue = Param->GetUnitDefault();
+						Parameter.MinValue = Param->GetUnitMin();
+						Parameter.MaxValue = Param->GetUnitMax();
+						return Parameter;
+					}
 				}
 			}
 		}
 
-		UE_LOG(LogAudioModulation, Error, TEXT("Audio modulation parameter '%s' not found. Modulation may be disabled for destination referencing parameter."), *InParamName.ToString());
+		UE_LOG(LogAudioModulation, Error, TEXT("Audio modulation parameter '%s' not found/loaded. Modulation may be disabled for destination referencing parameter."), *InParamName.ToString());
 		return Parameter;
 	}
 
diff --git a/Engine/Plugins/Runtime/AudioModulation/Source/AudioModulation/Public/AudioModulationSettings.h b/Engine/Plugins/Runtime/AudioModulation/Source/AudioModulation/Public/AudioModulationSettings.h
new file mode 100644
--- /dev/null
+++ b/Engine/Plugins/Runtime/AudioModulation/Source/AudioModulation/Public/AudioModulationSettings.h
@@ -0,0 +1,22 @@
+// Copyright Epic Games, Inc. All Rights Reserved.
+#pragma once
+
+#include "CoreMinimal.h"
+
+#include "Engine/DeveloperSettings.h"
+#include "UObject/ObjectMacros.h"
+#include "UObject/SoftObjectPath.h"
+
+#include "AudioModulationSettings.generated.h"
+
+
+UCLASS(config=AudioModulation, defaultconfig, meta = (DisplayName = "Audio Modulation"))
+class AUDIOMODULATION_API UAudioModulationSettings : public UDeveloperSettings
+{
+	GENERATED_BODY()
+
+public:
+	// Array of loaded Modulation Parameters
+	UPROPERTY(config, EditAnywhere, Category = "Parameters", meta = (AllowedClasses = "SoundModulationParameter"))
+	TArray<FSoftObjectPath> Parameters;
+};
\ No newline at end of file
