Fix missing font issue & visibility confusion in patch editor

diff --git a/Engine/Plugins/Runtime/AudioModulation/Source/AudioModulationEditor/Private/Editors/ModulationPatchCurveEditorViewStacked.cpp b/Engine/Plugins/Runtime/AudioModulation/Source/AudioModulationEditor/Private/Editors/ModulationPatchCurveEditorViewStacked.cpp
--- a/Engine/Plugins/Runtime/AudioModulation/Source/AudioModulationEditor/Private/Editors/ModulationPatchCurveEditorViewStacked.cpp
+++ b/Engine/Plugins/Runtime/AudioModulation/Source/AudioModulationEditor/Private/Editors/ModulationPatchCurveEditorViewStacked.cpp
@@ -215,7 +215,7 @@ void SModulationPatchEditorViewStacked::DrawLabels(const FGeometry& AllottedGeom
 	const double ValuePerPixel = 1.0 / StackedHeight;
 	const double ValueSpacePadding = StackedPadding * ValuePerPixel;
 
-	const FSlateFontInfo FontInfo = FCoreStyle::Get().GetFontStyle("FontAwesome.11");
+	const FSlateFontInfo FontInfo = FEditorStyle::GetFontStyle("CurveEd.LabelFont");
 	const FVector2D LocalSize = AllottedGeometry.GetLocalSize();
 	const FCurveEditorScreenSpace ViewSpace = GetViewSpace();
 
@@ -383,7 +383,7 @@ void SModulationPatchEditorViewStacked::DrawViewGridLineX(FSlateWindowElementLis
 	const double PixelTop = DrawInfo.GetPixelTop();
 
 	const TSharedRef<FSlateFontMeasure> FontMeasure = FSlateApplication::Get().GetRenderer()->GetFontMeasureService();
-	const FSlateFontInfo FontInfo = FCoreStyle::Get().GetFontStyle("ToolTip.LargerFont");
+	const FSlateFontInfo FontInfo = FCoreStyle::Get().GetFontStyle("CurveEd.InfoFont");
 
 	FLinearColor Color = bIsMajor ? DrawInfo.GetMajorGridColor() : DrawInfo.GetMinorGridColor();
 
@@ -444,7 +444,7 @@ void SModulationPatchEditorViewStacked::DrawViewGridLineY(const float VerticalLi
 
 		if (Label)
 		{
-			const FSlateFontInfo FontInfo = FCoreStyle::Get().GetFontStyle("ToolTip.LargerFont");
+			const FSlateFontInfo FontInfo = FCoreStyle::Get().GetFontStyle("CurveEd.InfoFont");
 
 			const TSharedRef<FSlateFontMeasure> FontMeasure = FSlateApplication::Get().GetRenderer()->GetFontMeasureService();
 			const FVector2D LabelSize = FontMeasure->Measure(*Label, FontInfo);
diff --git a/Engine/Plugins/Runtime/AudioModulation/Source/AudioModulationEditor/Private/Layouts/SoundControlModulationPatchLayout.cpp b/Engine/Plugins/Runtime/AudioModulation/Source/AudioModulationEditor/Private/Layouts/SoundControlModulationPatchLayout.cpp
--- a/Engine/Plugins/Runtime/AudioModulation/Source/AudioModulationEditor/Private/Layouts/SoundControlModulationPatchLayout.cpp
+++ b/Engine/Plugins/Runtime/AudioModulation/Source/AudioModulationEditor/Private/Layouts/SoundControlModulationPatchLayout.cpp
@@ -75,15 +75,16 @@ void FSoundControlModulationPatchLayoutCustomization::CustomizeChildren(TSharedR
 	TSharedRef<IPropertyHandle> InputsHandle = PropertyHandles.FindChecked(GET_MEMBER_NAME_CHECKED(FSoundControlModulationPatch, Inputs)).ToSharedRef();
 	TSharedRef<IPropertyHandle> ParameterHandle = PropertyHandles.FindChecked(GET_MEMBER_NAME_CHECKED(FSoundControlModulationPatch, OutputParameter)).ToSharedRef();
 
+	TAttribute<bool> EditCondition = TAttribute<bool>::Create([this, BypassHandle]()
+	{
+		bool bIsBypassed = false;
+		BypassHandle->GetValue(bIsBypassed);
+		return !bIsBypassed;
+	});
+
 	ChildBuilder.AddProperty(BypassHandle);
-	ChildBuilder.AddProperty(ParameterHandle);
-	ChildBuilder.AddProperty(InputsHandle)
-		.Visibility(TAttribute<EVisibility>::Create([this, BypassHandle]()
-		{
-			bool bIsBypassed = false;
-			BypassHandle->GetValue(bIsBypassed);
-			return bIsBypassed ? EVisibility::Hidden : EVisibility::Visible;
-		}));
+	ChildBuilder.AddProperty(ParameterHandle).EditCondition(EditCondition, nullptr);
+	ChildBuilder.AddProperty(InputsHandle).EditCondition(EditCondition, nullptr);
 }
 
 void FSoundControlModulationPatchLayoutCustomization::CustomizeHeader(TSharedRef<IPropertyHandle> StructPropertyHandle, FDetailWidgetRow& HeaderRow, IPropertyTypeCustomizationUtils& CustomizationUtils)
