Generalize LFO implementation so Modulation supports any type of control-rate generator

diff --git a/Engine/Plugins/Runtime/AudioModulation/Source/AudioModulation/Private/AudioModulation.cpp b/Engine/Plugins/Runtime/AudioModulation/Source/AudioModulation/Private/AudioModulation.cpp
--- a/Engine/Plugins/Runtime/AudioModulation/Source/AudioModulation/Private/AudioModulation.cpp
+++ b/Engine/Plugins/Runtime/AudioModulation/Source/AudioModulation/Private/AudioModulation.cpp
@@ -1,8 +1,8 @@
 // Copyright Epic Games, Inc. All Rights Reserved.
 #include "AudioModulation.h"
 
-#include "AudioModulationSystem.h"
 #include "AudioModulationLogging.h"
+#include "AudioModulationSystem.h"
 #include "CanvasTypes.h"
 #include "Features/IModularFeatures.h"
 #include "IAudioModulation.h"
@@ -18,13 +18,13 @@ DEFINE_STAT(STAT_AudioModulationProcessModulators);
 namespace AudioModulation
 {
 	FAudioModulation::FAudioModulation()
+		: ModSystem(new FAudioModulationSystem())
 	{
-		ModSystem = MakeUnique<FAudioModulationSystem>();
 	}
 
-	FAudioModulationSystem* FAudioModulation::GetModulationSystem()
+	FAudioModulation::~FAudioModulation()
 	{
-		return ModSystem.Get();
+		delete ModSystem;
 	}
 
 	Audio::FModulationParameter FAudioModulation::GetParameter(FName InParamName)
@@ -42,6 +42,71 @@ namespace AudioModulation
 		ModSystem->OnAuditionEnd();
 	}
 
+	void FAudioModulation::ActivateBus(const USoundControlBus& InBus)
+	{
+		ModSystem->ActivateBus(InBus);
+	}
+
+	void FAudioModulation::ActivateBusMix(const USoundControlBusMix& InBusMix)
+	{
+		ModSystem->ActivateBusMix(InBusMix);
+	}
+
+	void FAudioModulation::ActivateGenerator(const USoundModulationGenerator& InGenerator)
+	{
+		ModSystem->ActivateGenerator(InGenerator);
+	}
+
+	void FAudioModulation::DeactivateBus(const USoundControlBus& InBus)
+	{
+		ModSystem->DeactivateBus(InBus);
+	}
+
+	void FAudioModulation::DeactivateBusMix(const USoundControlBusMix& InBusMix)
+	{
+		ModSystem->SoloBusMix(InBusMix);
+	}
+
+	void FAudioModulation::DeactivateAllBusMixes()
+	{
+		ModSystem->DeactivateAllBusMixes();
+	}
+
+	void FAudioModulation::DeactivateGenerator(const USoundModulationGenerator& InGenerator)
+	{
+		ModSystem->DeactivateGenerator(InGenerator);
+	}
+
+	void FAudioModulation::SaveMixToProfile(const USoundControlBusMix& InBusMix, const int32 InProfileIndex)
+	{
+		ModSystem->SaveMixToProfile(InBusMix, InProfileIndex);
+	}
+
+	TArray<FSoundControlBusMixStage> FAudioModulation::LoadMixFromProfile(const int32 InProfileIndex, USoundControlBusMix& OutBusMix)
+	{
+		return ModSystem->LoadMixFromProfile(InProfileIndex, OutBusMix);
+	}
+
+	void FAudioModulation::UpdateMix(const TArray<FSoundControlBusMixStage>& InStages, USoundControlBusMix& InOutMix, bool bInUpdateObject, float InFadeTime)
+	{
+		ModSystem->UpdateMix(InStages, InOutMix, bInUpdateObject, InFadeTime);
+	}
+
+	void FAudioModulation::UpdateMix(const USoundControlBusMix& InMix, float InFadeTime)
+	{
+		ModSystem->UpdateMix(InMix, InFadeTime);
+	}
+
+	void FAudioModulation::UpdateMixByFilter(const FString& InAddressFilter, const TSubclassOf<USoundModulationParameter>& InParamClassFilter, USoundModulationParameter* InParamFilter, float Value, float FadeTime, USoundControlBusMix& InOutMix, bool bInUpdateObject)
+	{
+		ModSystem->UpdateMixByFilter(InAddressFilter, InParamClassFilter, InParamFilter, Value, FadeTime, InOutMix, bInUpdateObject);
+	}
+
+	void FAudioModulation::SoloBusMix(const USoundControlBusMix& InBusMix)
+	{
+		ModSystem->SoloBusMix(InBusMix);
+	}
+
 #if !UE_BUILD_SHIPPING
 	bool FAudioModulation::OnPostHelp(FCommonViewportClient* ViewportClient, const TCHAR* Stream)
 	{
diff --git a/Engine/Plugins/Runtime/AudioModulation/Source/AudioModulation/Private/AudioModulationDebugger.cpp b/Engine/Plugins/Runtime/AudioModulation/Source/AudioModulation/Private/AudioModulationDebugger.cpp
--- a/Engine/Plugins/Runtime/AudioModulation/Source/AudioModulation/Private/AudioModulationDebugger.cpp
+++ b/Engine/Plugins/Runtime/AudioModulation/Source/AudioModulation/Private/AudioModulationDebugger.cpp
@@ -10,14 +10,13 @@
 #include "AudioThread.h"
 #include "Engine/Canvas.h"
 #include "Engine/Engine.h"
+#include "HAL/IConsoleManager.h"
 #include "IAudioModulation.h"
 #include "Misc/CoreDelegates.h"
 #include "SoundControlBusProxy.h"
 #include "SoundModulationPatch.h"
 #include "SoundModulationProxy.h"
 #include "SoundModulationValue.h"
-#include "HAL/IConsoleManager.h"
-
 
 
 static float AudioModulationDebugUpdateRateCVar = 0.25f;
@@ -28,6 +27,7 @@ FAutoConsoleVariableRef CVarAudioModulationDebugUpdateRate(
 	TEXT("Default: 0.25f"),
 	ECVF_Default);
 
+
 namespace AudioModulation
 {
 	namespace Debug
@@ -42,95 +42,70 @@ namespace AudioModulation
 				: FColor::Green;
 		}
 
-		int32 RenderStatLFO(const TArray<FLFODebugInfo>& FilteredLFOs, FCanvas& Canvas, int32 X, int32 Y, const UFont& Font)
+		int32 RenderStatGenerators(const FGeneratorDebugInfo& GeneratorDebugInfo, FCanvas& Canvas, int32 X, int32 Y, const UFont& Font)
 		{
 			int32 Height = 12;
 			int32 Width = 12;
 			Font.GetStringHeightAndWidth(TEXT("@"), Height, Width);
 
-			// Determine minimum width of cells
-			static const FString NameHeader      = TEXT("Name");
-			static const FString ValueHeader     = TEXT("Value");
-			static const FString AmplitudeHeader = TEXT("Amplitude");
-			static const FString FrequencyHeader = TEXT("Frequency");
-			static const FString OffsetHeader    = TEXT("Offset");
-			static const FString TypeHeader		= TEXT("Type");
-			static const TArray<int32> StaticCellWidths =
+			int32 CellWidth = 0;
+			for (const FString& Category : GeneratorDebugInfo.Categories)
 			{
-				NameHeader.Len(),
-				ValueHeader.Len(),
-				AmplitudeHeader.Len(),
-				FrequencyHeader.Len(),
-				OffsetHeader.Len(),
-				TypeHeader.Len()
-			};
+				CellWidth = FMath::Max(CellWidth, Category.Len());
+			}
 
-			int32 CellWidth = FMath::Max(StaticCellWidths);
-			for (const FLFODebugInfo& LFODebugInfo : FilteredLFOs)
+			for (const FGeneratorDebugInfo::FInstanceValues& InstanceInfo : GeneratorDebugInfo.FilteredInstances)
 			{
-				CellWidth = FMath::Max(CellWidth, LFODebugInfo.Name.Len());
+				for (const FString& Value : InstanceInfo)
+				{
+					CellWidth = FMath::Max(CellWidth, Value.Len());
+				}
 			}
 
-			Canvas.DrawShadowedString(X + XIndent, Y, TEXT("Active LFOs:"), &Font, FColor::Red);
+			Canvas.DrawShadowedString(X + XIndent, Y, TEXT("Active Generators:"), &Font, FColor::Red);
 			Y += Height;
 			X += XIndent;
 
-			int32 RowX = X;
-			Canvas.DrawShadowedString(RowX, Y, *NameHeader, &Font, FColor::White);
-			RowX += Width * CellWidth;
-
-			Canvas.DrawShadowedString(RowX, Y, *ValueHeader, &Font, FColor::White);
-			RowX += Width * CellWidth;
-
-			Canvas.DrawShadowedString(RowX, Y, *AmplitudeHeader, &Font, FColor::White);
-			RowX += Width * CellWidth;
-
-			Canvas.DrawShadowedString(RowX, Y, *FrequencyHeader, &Font, FColor::White);
-			RowX += Width * CellWidth;
-
-			Canvas.DrawShadowedString(RowX, Y, *OffsetHeader, &Font, FColor::White);
-			RowX += Width * CellWidth;
-
-			Canvas.DrawShadowedString(RowX, Y, *TypeHeader, &Font, FColor::White);
-			RowX += Width * CellWidth;
-
-			Y += Height;
-			for (const FLFODebugInfo& LFODebugInfo : FilteredLFOs)
+			// Print Category Titles
 			{
-				RowX = X;
+				int32 RowX = X;
 
-				FString Name = LFODebugInfo.Name.Left(MaxNameLength);
-				Name += FString::Printf(TEXT(" (%u)"), LFODebugInfo.RefCount);
-
-				if (Name.Len() < MaxNameLength)
-				{
-					Name = Name.RightPad(MaxNameLength - Name.Len());
-				}
-
-				Canvas.DrawShadowedString(RowX, Y, *Name, &Font, FColor::Green);
+				Canvas.DrawShadowedString(RowX, Y, TEXT("Name"), &Font, FColor::White);
 				RowX += Width * CellWidth;
 
-				const FString ValueString = FString::Printf(TEXT("%.6f"), LFODebugInfo.Value);
-				Canvas.DrawShadowedString(RowX, Y, *ValueString, &Font, FColor::Green);
+				Canvas.DrawShadowedString(RowX, Y, TEXT("Ref Count"), &Font, FColor::White);
 				RowX += Width * CellWidth;
 
-				const FString AmpString = FString::Printf(TEXT("%.3f"), LFODebugInfo.Amplitude);
-				Canvas.DrawShadowedString(RowX, Y, *AmpString, &Font, FColor::Green);
-				RowX += Width * CellWidth;
+				for (const FString& Category : GeneratorDebugInfo.Categories)
+				{
+					Canvas.DrawShadowedString(RowX, Y, *Category, &Font, FColor::White);
+					RowX += Width * CellWidth;
+				}
+				Y += Height;
+			}
 
-				const FString FreqString = FString::Printf(TEXT("%.4f"), LFODebugInfo.Frequency);
-				Canvas.DrawShadowedString(RowX, Y, *FreqString, &Font, FColor::Green);
-				RowX += Width * CellWidth;
+			// Print Instance Info
+			for (const FGeneratorDebugInfo::FInstanceValues& InstanceValues : GeneratorDebugInfo.FilteredInstances)
+			{
+				int32 RowX = X;
 
-				const FString OffsetString = FString::Printf(TEXT("%.4f"), LFODebugInfo.Offset);
-				Canvas.DrawShadowedString(RowX, Y, *OffsetString, &Font, FColor::Green);
-				RowX += Width * CellWidth;
+				// Validate that all instances have the same number of values as there are registered categories.
+				// 2 "extra" are from parent generator proxy that include proxy's Name & Ref Count.
+				ensure(InstanceValues.Num() == GeneratorDebugInfo.Categories.Num() + 2);
 
-				Canvas.DrawShadowedString(RowX, Y, *LFODebugInfo.Type, &Font, FColor::Green);
-				RowX += Width * CellWidth;
+				for (const FString& Value : InstanceValues)
+				{
+					FString DisplayValue = Value.Left(MaxNameLength);
+					if (DisplayValue.Len() < MaxNameLength)
+					{
+						DisplayValue = DisplayValue.RightPad(MaxNameLength - DisplayValue.Len());
+					}
 
-				Y += Height;
+					Canvas.DrawShadowedString(RowX, Y, *DisplayValue, &Font, FColor::Green);
+					RowX += Width * CellWidth;
+				}
 			}
+			Y += Height;
 
 			return Y;
 		}
@@ -147,11 +122,11 @@ namespace AudioModulation
 
 			// Determine minimum width of cells
 			static const FString MixSubTotalHeader = TEXT("Mix");
-			static const FString LFOSubTotalHeader = TEXT("LFO");
-			static const FString TotalHeader       = TEXT("Final");
+			static const FString GeneratorsSubTotalHeader = TEXT("Generators");
+			static const FString TotalHeader = TEXT("Final");
 			static const TArray<int32> StaticCellWidths =
 			{
-				LFOSubTotalHeader.Len(),
+				GeneratorsSubTotalHeader.Len(),
 				MixSubTotalHeader.Len(),
 				TotalHeader.Len(),
 				14 /* Minimum width due to value strings X.XXXX(X.XXXX) */
@@ -252,12 +227,12 @@ namespace AudioModulation
 			}
 			Y += Height;
 
-			Canvas.DrawShadowedString(X, Y, *LFOSubTotalHeader, &Font, FColor::Yellow);
+			Canvas.DrawShadowedString(X, Y, *GeneratorsSubTotalHeader, &Font, FColor::Yellow);
 			RowX = X;
 			for (const FControlBusDebugInfo& Bus : FilteredBuses)
 			{
 				RowX += Width * CellWidth; // Add before to leave space for row headers
-				const float Value = Bus.LFOValue;
+				const float Value = Bus.GeneratorValue;
 				const FColor Color = Debug::GetUnitRenderColor(Value);
 				Canvas.DrawShadowedString(RowX, Y, *FString::Printf(TEXT("%.4f"), Value), &Font, Color);
 			}
@@ -305,7 +280,6 @@ namespace AudioModulation
 
 	FAudioModulationDebugger::FAudioModulationDebugger()
 		: bActive(0)
-		, bShowRenderStatLFO(1)
 		, bShowRenderStatMix(1)
 		, ElapsedSinceLastUpdate(0.0f)
 	{
@@ -335,7 +309,7 @@ namespace AudioModulation
 			FControlBusDebugInfo DebugInfo;
 			DebugInfo.DefaultValue = Proxy->GetDefaultValue();
 			DebugInfo.Id = Proxy->GetId();
-			DebugInfo.LFOValue = Proxy->GetLFOValue();
+			DebugInfo.GeneratorValue = Proxy->GetGeneratorValue();
 			DebugInfo.MixValue = Proxy->GetMixValue();
 			DebugInfo.Name = Proxy->GetName();
 			DebugInfo.RefCount = Proxy->GetRefCount();
@@ -362,60 +336,18 @@ namespace AudioModulation
 			RefreshedFilteredMixes.Add(DebugInfo);
 		}
 
-		static const int32 MaxFilteredLFOs = 8;
-		TArray<const FModulatorLFOProxy*> FilteredLFOProxies;
-		Debug::FilterDebugArray<FLFOProxyMap, FModulatorLFOProxy>(InRefProxies.LFOs, LFOStringFilter, MaxFilteredLFOs, FilteredLFOProxies);
+		static const int32 MaxFilteredGenerators = 8;
+		TArray<const FModulatorGeneratorProxy*> FilteredGeneratorProxies;
+		Debug::FilterDebugArray<FGeneratorProxyMap, FModulatorGeneratorProxy>(InRefProxies.Generators, GeneratorStringFilter, MaxFilteredGenerators, FilteredGeneratorProxies);
 
-		TArray<FLFODebugInfo> RefreshedFilteredLFOs;
-		for (const FModulatorLFOProxy* Proxy : FilteredLFOProxies)
+		FGeneratorSortMap RefreshedFilteredGenerators;
+		for (const FModulatorGeneratorProxy* Proxy : FilteredGeneratorProxies)
 		{
-			FLFODebugInfo DebugInfo;
-			DebugInfo.Name = Proxy->GetName();
-			DebugInfo.RefCount = Proxy->GetRefCount();
-			DebugInfo.Value = Proxy->GetValue();
-			DebugInfo.Amplitude = Proxy->GetLFO().GetGain();
-			DebugInfo.Frequency = Proxy->GetLFO().GetFrequency();
-			DebugInfo.Offset = Proxy->GetOffset();
-
-			switch (Proxy->GetLFO().GetType())
-			{
-				case Audio::ELFO::DownSaw:
-				DebugInfo.Type = TEXT("DownSaw");
-				break;
-
-				case Audio::ELFO::Exponential:
-				DebugInfo.Type = TEXT("Exponential");
-				break;
-
-				case Audio::ELFO::RandomSampleHold:
-				DebugInfo.Type = TEXT("Random (Sample & Hold)");
-				break;
-
-				case Audio::ELFO::Sine:
-				DebugInfo.Type = TEXT("Sine");
-				break;
-
-				case Audio::ELFO::Square:
-				DebugInfo.Type = TEXT("Square");
-				break;
-
-				case Audio::ELFO::Triangle:
-				DebugInfo.Type = TEXT("Triangle");
-				break;
-
-				case Audio::ELFO::UpSaw:
-				DebugInfo.Type = TEXT("Up Saw");
-				break;
-
-				default:
-				static_assert(static_cast<int32>(Audio::ELFO::NumLFOTypes) == 7, "Missing LFO type case coverage");
-				break;
-			}
-
-			RefreshedFilteredLFOs.Add(DebugInfo);
+			FGeneratorDebugInfo::FInstanceValues Values = Proxy->GetDebugValues();
+			RefreshedFilteredGenerators.FindOrAdd(Proxy->GetDebugName()).FilteredInstances.Emplace(MoveTemp(Values));
 		}
 
-		FFunctionGraphTask::CreateAndDispatchWhenReady([this, RefreshedFilteredBuses, RefreshedFilteredMixes, RefreshedFilteredLFOs]()
+		FFunctionGraphTask::CreateAndDispatchWhenReady([this, RefreshedFilteredBuses, RefreshedFilteredMixes, RefreshedFilteredGenerators]()
 		{
 			FilteredBuses = RefreshedFilteredBuses;
 			FilteredBuses.Sort(&Debug::CompareNames<FControlBusDebugInfo>);
@@ -423,8 +355,20 @@ namespace AudioModulation
 			FilteredMixes = RefreshedFilteredMixes;
 			FilteredMixes.Sort(&Debug::CompareNames<FControlBusMixDebugInfo>);
 
-			FilteredLFOs = RefreshedFilteredLFOs;
-			FilteredLFOs.Sort(&Debug::CompareNames<FLFODebugInfo>);
+			FGeneratorSortMap NewFilteredMap = RefreshedFilteredGenerators;
+			for (TPair<FString, FGeneratorDebugInfo>& NewInfoPair : NewFilteredMap)
+			{
+				if (FGeneratorDebugInfo* LastDebugInfo = FilteredGeneratorsMap.Find(NewInfoPair.Key))
+				{
+					NewInfoPair.Value.bEnabled = LastDebugInfo->bEnabled;
+					if (!NewInfoPair.Value.bEnabled)
+					{
+						NewInfoPair.Value.FilteredInstances.Reset();
+					}
+				}
+			}
+
+			FilteredGeneratorsMap = MoveTemp(NewFilteredMap);
 		}, TStatId(), nullptr, ENamedThreads::GameThread);
 	}
 
@@ -436,8 +380,8 @@ namespace AudioModulation
 			GEngine->AddOnScreenDebugMessage(-1, TimeShowHelp, FColor::Yellow, TEXT("-MixFilter: Substring that filters mixes shown in matrix view"));
 			GEngine->AddOnScreenDebugMessage(-1, TimeShowHelp, FColor::Yellow, TEXT("-BusFilter: Substring that filters buses shown in matrix view"));
 			GEngine->AddOnScreenDebugMessage(-1, TimeShowHelp, FColor::Yellow, TEXT("-Matrix: Show bus matrix"));
-			GEngine->AddOnScreenDebugMessage(-1, TimeShowHelp, FColor::Yellow, TEXT("-LFOFilter: Substring that filters LFOs listed"));
-			GEngine->AddOnScreenDebugMessage(-1, TimeShowHelp, FColor::Yellow, TEXT("-LFO: Show LFO debug data"));
+			GEngine->AddOnScreenDebugMessage(-1, TimeShowHelp, FColor::Yellow, TEXT("-GeneratorFilter: Substring that filters Generator types listed"));
+			GEngine->AddOnScreenDebugMessage(-1, TimeShowHelp, FColor::Yellow, TEXT("-Generator: Show Generator debug data"));
 			GEngine->AddOnScreenDebugMessage(-1, TimeShowHelp, FColor::Yellow, TEXT("stat SoundModulators:"));
 		}
 
@@ -460,9 +404,12 @@ namespace AudioModulation
 			Y = Debug::RenderStatMixMatrix(FilteredMixes, FilteredBuses, Canvas, X, Y, Font);
 		}
 
-		if (bShowRenderStatLFO)
+		for (const TPair<FString, FGeneratorDebugInfo>& GeneratorStatPair : FilteredGeneratorsMap)
 		{
-			Y = Debug::RenderStatLFO(FilteredLFOs, Canvas, X, Y, Font);
+			if (GeneratorStatPair.Value.bEnabled)
+			{
+				Y = Debug::RenderStatGenerators(GeneratorStatPair.Value, Canvas, X, Y, Font);
+			}
 		}
 
 		return Y;
@@ -481,31 +428,46 @@ namespace AudioModulation
 		{
 			if (Stream && Stream[0] != '\0')
 			{
-				const bool bUpdateLFO = FParse::Param(Stream, TEXT("LFO"));
-				const bool bUpdateMix = FParse::Param(Stream, TEXT("Matrix"));
+				bool bGeneratorsUpdated = false;
+				for (TPair<FString, FGeneratorDebugInfo>& GeneratorStatPair : FilteredGeneratorsMap)
+				{
+					const bool bGeneratorEnabled = FParse::Param(Stream, *GeneratorStatPair.Key);
+					if (bGeneratorEnabled)
+					{
+						bGeneratorsUpdated = true;
+						GeneratorStatPair.Value.bEnabled = bGeneratorEnabled;
+					}
+					else
+					{
+						GeneratorStatPair.Value.bEnabled = false;
+					}
+				}
 
-				if (bUpdateLFO || bUpdateMix)
+				const bool bUpdateMix = FParse::Param(Stream, TEXT("Matrix"));
+				if (bGeneratorsUpdated || bUpdateMix)
 				{
-					bShowRenderStatLFO = bUpdateLFO;
 					bShowRenderStatMix = bUpdateMix;
 				}
 				else
 				{
-					bShowRenderStatLFO = 1;
+					// Off by default as there may be many of these
+					ResetGeneratorStats();
 					bShowRenderStatMix = 1;
 				}
 
 				FParse::Value(Stream, TEXT("BusFilter"), BusStringFilter);
-				FParse::Value(Stream, TEXT("LFOFilter"), LFOStringFilter);
+				FParse::Value(Stream, TEXT("GeneratorFilter"), GeneratorStringFilter);
 				FParse::Value(Stream, TEXT("MixFilter"), MixStringFilter);
 			}
 			else
 			{
-				bShowRenderStatLFO = 1;
 				bShowRenderStatMix = 1;
 
+				// Off by default as there may be many of these
+				ResetGeneratorStats();
+
 				BusStringFilter.Reset();
-				LFOStringFilter.Reset();
+				GeneratorStringFilter.Reset();
 				MixStringFilter.Reset();
 			}
 		}
@@ -513,6 +475,15 @@ namespace AudioModulation
 		bActive = !bActive;
 		return true;
 	}
+
+	void FAudioModulationDebugger::ResetGeneratorStats()
+	{
+		for (TPair<FString, FGeneratorDebugInfo>& GeneratorStatPair : FilteredGeneratorsMap)
+		{
+			GeneratorStatPair.Value.bEnabled = false;
+		}
+	}
+
 } // namespace AudioModulation
 #endif // !UE_BUILD_SHIPPING
 #endif // WITH_AUDIOMODULATION
\ No newline at end of file
diff --git a/Engine/Plugins/Runtime/AudioModulation/Source/AudioModulation/Private/AudioModulationDebugger.h b/Engine/Plugins/Runtime/AudioModulation/Source/AudioModulation/Private/AudioModulationDebugger.h
--- a/Engine/Plugins/Runtime/AudioModulation/Source/AudioModulation/Private/AudioModulationDebugger.h
+++ b/Engine/Plugins/Runtime/AudioModulation/Source/AudioModulation/Private/AudioModulationDebugger.h
@@ -7,8 +7,8 @@
 #if !UE_BUILD_SHIPPING
 #include "SoundControlBus.h"
 #include "SoundControlBusMix.h"
+#include "SoundModulationGenerator.h"
 #include "SoundModulationProxy.h"
-#include "SoundModulationGeneratorLFO.h"
 
 
 namespace AudioModulation
@@ -35,22 +35,37 @@ namespace AudioModulation
 	{
 		FString Name;
 		float DefaultValue;
-		float LFOValue;
+		float GeneratorValue;
 		float MixValue;
 		float Value;
 		uint32 Id;
 		uint32 RefCount;
 	};
 
-	struct FLFODebugInfo
+	struct FGeneratorSort
 	{
-		FString Name;
-		float Value;
-		float Amplitude;
-		float Frequency;
-		float Offset;
-		FString Type;
-		uint32 RefCount;
+		FORCEINLINE bool operator()(const FString& A, const FString& B) const
+		{
+			return A < B;
+		}
+	};
+
+	struct FGeneratorDebugInfo
+	{
+		FGeneratorDebugInfo()
+		{
+		}
+
+		FGeneratorDebugInfo(const TArray<FString>& InCategories)
+			: Categories(InCategories)
+		{
+		}
+
+		bool bEnabled = false;
+		TArray<FString> Categories;
+
+		using FInstanceValues = TArray<FString>;
+		TArray<FInstanceValues> FilteredInstances;
 	};
 
 	class FAudioModulationDebugger
@@ -63,21 +78,25 @@ namespace AudioModulation
 		int32 OnRenderStat(FCanvas& Canvas, int32 X, int32 Y, const UFont& Font);
 		bool OnToggleStat(FCommonViewportClient& ViewportClient, const TCHAR* Stream);
 
+		void ResetGeneratorStats();
+
 	private:
 		uint8 bActive : 1;
-		uint8 bShowRenderStatLFO : 1;
 		uint8 bShowRenderStatMix : 1;
 
-		TArray<FControlBusDebugInfo>    FilteredBuses;
-		TArray<FLFODebugInfo>           FilteredLFOs;
+		TArray<FControlBusDebugInfo> FilteredBuses;
 		TArray<FControlBusMixDebugInfo> FilteredMixes;
 
+		using FGeneratorSortMap = TSortedMap<FString, FGeneratorDebugInfo, FDefaultAllocator, FGeneratorSort>;
+		FGeneratorSortMap FilteredGeneratorsMap;
+
 		FString BusStringFilter;
-		FString LFOStringFilter;
+		FString GeneratorStringFilter;
 		FString MixStringFilter;
 
 		float ElapsedSinceLastUpdate;
 	};
+
 } // namespace AudioModulation
 #endif // !UE_BUILD_SHIPPING
 #endif // WITH_AUDIOMODULATION
diff --git a/Engine/Plugins/Runtime/AudioModulation/Source/AudioModulation/Private/AudioModulationStatics.cpp b/Engine/Plugins/Runtime/AudioModulation/Source/AudioModulation/Private/AudioModulationStatics.cpp
--- a/Engine/Plugins/Runtime/AudioModulation/Source/AudioModulation/Private/AudioModulationStatics.cpp
+++ b/Engine/Plugins/Runtime/AudioModulation/Source/AudioModulation/Private/AudioModulationStatics.cpp
@@ -117,7 +117,7 @@ void UAudioModulationStatics::ActivateBus(const UObject* WorldContextObject, USo
 	}
 
 	UWorld* World = GetAudioWorld(WorldContextObject);
-	if (AudioModulation::FAudioModulationSystem* ModSystem = GetModulationSystem(World))
+	if (AudioModulation::FAudioModulation* ModSystem = GetModulation(World))
 	{
 		ModSystem->ActivateBus(*Bus);
 	}
@@ -128,21 +128,21 @@ void UAudioModulationStatics::ActivateBusMix(const UObject* WorldContextObject,
 	if (BusMix)
 	{
 		UWorld* World = GetAudioWorld(WorldContextObject);
-		if (AudioModulation::FAudioModulationSystem* ModSystem = GetModulationSystem(World))
+		if (AudioModulation::FAudioModulation* ModSystem = GetModulation(World))
 		{
 			ModSystem->ActivateBusMix(*BusMix);
 		}
 	}
 }
 
-void UAudioModulationStatics::ActivateBusModulator(const UObject* WorldContextObject, USoundModulationGenerator* Modulator)
+void UAudioModulationStatics::ActivateGenerator(const UObject* WorldContextObject, USoundModulationGenerator* Generator)
 {
 	UWorld* World = GetAudioWorld(WorldContextObject);
-	if (AudioModulation::FAudioModulationSystem* ModSystem = GetModulationSystem(World))
+	if (AudioModulation::FAudioModulation* ModSystem = GetModulation(World))
 	{
-		if (USoundModulationGeneratorLFO* LFO = Cast<USoundModulationGeneratorLFO>(Modulator))
+		if (Generator)
 		{
-			ModSystem->ActivateLFO(*LFO);
+			ModSystem->ActivateGenerator(*Generator);
 		}
 	}
 }
@@ -163,7 +163,7 @@ UWorld* UAudioModulationStatics::GetAudioWorld(const UObject* WorldContextObject
 	return World;
 }
 
-AudioModulation::FAudioModulationSystem* UAudioModulationStatics::GetModulationSystem(UWorld* World)
+AudioModulation::FAudioModulation* UAudioModulationStatics::GetModulation(UWorld* World)
 {
 	FAudioDeviceHandle AudioDevice;
 	if (World)
@@ -182,7 +182,7 @@ AudioModulation::FAudioModulationSystem* UAudioModulationStatics::GetModulationS
 	{
 		if (IAudioModulation* ModulationInterface = AudioDevice->ModulationInterface.Get())
 		{
-			return static_cast<AudioModulation::FAudioModulation*>(ModulationInterface)->GetModulationSystem();
+			return static_cast<AudioModulation::FAudioModulation*>(ModulationInterface);
 		}
 	}
 
@@ -203,7 +203,7 @@ USoundControlBus* UAudioModulationStatics::CreateBus(const UObject* WorldContext
 
 	if (Activate)
 	{
-		if (AudioModulation::FAudioModulationSystem* ModSystem = UAudioModulationStatics::GetModulationSystem(World))
+		if (AudioModulation::FAudioModulation* ModSystem = UAudioModulationStatics::GetModulation(World))
 		{
 			ModSystem->ActivateBus(*NewBus);
 		}
@@ -212,30 +212,6 @@ USoundControlBus* UAudioModulationStatics::CreateBus(const UObject* WorldContext
 	return NewBus;
 }
 
-USoundModulationGeneratorLFO* UAudioModulationStatics::CreateLFO(const UObject* WorldContextObject, FName Name, float Amplitude, float Frequency, float Offset, bool Activate)
-{
-	UWorld* World = GetAudioWorld(WorldContextObject);
-	if (!World)
-	{
-		return nullptr;
-	}
-
-	USoundModulationGeneratorLFO* NewLFO = NewObject<USoundModulationGeneratorLFO>(GetTransientPackage(), Name);
-	NewLFO->Amplitude = Amplitude;
-	NewLFO->Frequency = Frequency;
-	NewLFO->Offset    = Offset;
-
-	if (Activate)
-	{
-		if (AudioModulation::FAudioModulationSystem* ModSystem = GetModulationSystem(World))
-		{
-			ModSystem->ActivateLFO(*NewLFO);
-		}
-	}
-
-	return NewLFO;
-}
-
 FSoundControlBusMixStage UAudioModulationStatics::CreateBusMixStage(const UObject* WorldContextObject, USoundControlBus* Bus, float Value, float AttackTime, float ReleaseTime)
 {
 	FSoundControlBusMixStage MixStage;
@@ -269,7 +245,7 @@ USoundControlBusMix* UAudioModulationStatics::CreateBusMix(const UObject* WorldC
 
 	if (Activate)
 	{
-		if (AudioModulation::FAudioModulationSystem* ModSystem = GetModulationSystem(World))
+		if (AudioModulation::FAudioModulation* ModSystem = GetModulation(World))
 		{
 			ModSystem->ActivateBusMix(*NewBusMix);
 		}
@@ -283,7 +259,7 @@ void UAudioModulationStatics::DeactivateBus(const UObject* WorldContextObject, U
 	if (Bus)
 	{
 		UWorld* World = GetAudioWorld(WorldContextObject);
-		if (AudioModulation::FAudioModulationSystem* ModSystem = GetModulationSystem(World))
+		if (AudioModulation::FAudioModulation* ModSystem = GetModulation(World))
 		{
 			ModSystem->DeactivateBus(*Bus);
 		}
@@ -295,21 +271,21 @@ void UAudioModulationStatics::DeactivateBusMix(const UObject* WorldContextObject
 	if (BusMix)
 	{
 		UWorld* World = GetAudioWorld(WorldContextObject);
-		if (AudioModulation::FAudioModulationSystem* ModSystem = GetModulationSystem(World))
+		if (AudioModulation::FAudioModulation* ModSystem = GetModulation(World))
 		{
 			ModSystem->DeactivateBusMix(*BusMix);
 		}
 	}
 }
 
-void UAudioModulationStatics::DeactivateBusModulator(const UObject* WorldContextObject, USoundModulationGenerator* Modulator)
+void UAudioModulationStatics::DeactivateGenerator(const UObject* WorldContextObject, USoundModulationGenerator* Generator)
 {
 	UWorld* World = GetAudioWorld(WorldContextObject);
-	if (AudioModulation::FAudioModulationSystem* ModSystem = GetModulationSystem(World))
+	if (AudioModulation::FAudioModulation* ModSystem = GetModulation(World))
 	{
-		if (USoundModulationGeneratorLFO* LFO = Cast<USoundModulationGeneratorLFO>(Modulator))
+		if (Generator)
 		{
-			ModSystem->DeactivateLFO(*LFO);
+			ModSystem->DeactivateGenerator(*Generator);
 		}
 	}
 }
@@ -317,7 +293,7 @@ void UAudioModulationStatics::DeactivateBusModulator(const UObject* WorldContext
 void UAudioModulationStatics::SaveMixToProfile(const UObject* WorldContextObject, USoundControlBusMix* BusMix, int32 ProfileIndex)
 {
 	UWorld* World = GetAudioWorld(WorldContextObject);
-	if (AudioModulation::FAudioModulationSystem* ModSystem = GetModulationSystem(World))
+	if (AudioModulation::FAudioModulation* ModSystem = GetModulation(World))
 	{
 		if (BusMix)
 		{
@@ -331,7 +307,7 @@ TArray<FSoundControlBusMixStage> UAudioModulationStatics::LoadMixFromProfile(con
 	if (BusMix)
 	{
 		UWorld* World = GetAudioWorld(WorldContextObject);
-		if (AudioModulation::FAudioModulationSystem* ModSystem = GetModulationSystem(World))
+		if (AudioModulation::FAudioModulation* ModSystem = GetModulation(World))
 		{
 			if (bActivate)
 			{
@@ -349,7 +325,7 @@ void UAudioModulationStatics::UpdateMix(const UObject* WorldContextObject, USoun
 	if (Mix)
 	{
 		UWorld* World = GetAudioWorld(WorldContextObject);
-		if (AudioModulation::FAudioModulationSystem* ModSystem = GetModulationSystem(World))
+		if (AudioModulation::FAudioModulation* ModSystem = GetModulation(World))
 		{
 			// UObject representation is not updated in this form of the call as doing so from
 			// PIE can result in an unstable state where UObject is modified but not properly dirtied.
@@ -370,7 +346,7 @@ void UAudioModulationStatics::UpdateMixByFilter(
 	if (Mix)
 	{
 		UWorld* World = GetAudioWorld(WorldContextObject);
-		if (AudioModulation::FAudioModulationSystem* ModSystem = GetModulationSystem(World))
+		if (AudioModulation::FAudioModulation* ModSystem = GetModulation(World))
 		{
 			// UObject representation is not updated in this form of the call as doing so from
 			// PIE can result in an unstable state where UObject is modified but not properly dirtied.
@@ -384,7 +360,7 @@ void UAudioModulationStatics::UpdateMixFromObject(const UObject* WorldContextObj
 	if (Mix)
 	{
 		UWorld* World = GetAudioWorld(WorldContextObject);
-		if (AudioModulation::FAudioModulationSystem* ModSystem = GetModulationSystem(World))
+		if (AudioModulation::FAudioModulation* ModSystem = GetModulation(World))
 		{
 			ModSystem->UpdateMix(*Mix, InFadeTime);
 		}
@@ -396,7 +372,7 @@ void UAudioModulationStatics::UpdateModulator(const UObject* WorldContextObject,
 	if (Modulator)
 	{
 		UWorld* World = GetAudioWorld(WorldContextObject);
-		if (AudioModulation::FAudioModulationSystem* ModSystem = GetModulationSystem(World))
+		if (AudioModulation::FAudioModulation* ModSystem = GetModulation(World))
 		{
 			ModSystem->UpdateModulator(*Modulator);
 		}
diff --git a/Engine/Plugins/Runtime/AudioModulation/Source/AudioModulation/Private/AudioModulationSystem.cpp b/Engine/Plugins/Runtime/AudioModulation/Source/AudioModulation/Private/AudioModulationSystem.cpp
--- a/Engine/Plugins/Runtime/AudioModulation/Source/AudioModulation/Private/AudioModulationSystem.cpp
+++ b/Engine/Plugins/Runtime/AudioModulation/Source/AudioModulation/Private/AudioModulationSystem.cpp
@@ -2,6 +2,11 @@
 #include "AudioModulationSystem.h"
 
 #if WITH_AUDIOMODULATION
+
+#if !UE_BUILD_SHIPPING
+#include "AudioModulationDebugger.h"
+#endif // !UE_BUILD_SHIPPING
+
 #include "Async/Async.h"
 #include "Audio/AudioAddressPattern.h"
 #include "AudioModulationLogging.h"
@@ -10,24 +15,20 @@
 #include "AudioThread.h"
 #include "CoreGlobals.h"
 #include "Engine/Engine.h"
+#include "Generators/SoundModulationLFO.h"
 #include "HAL/PlatformTLS.h"
 #include "IAudioModulation.h"
 #include "Misc/CoreDelegates.h"
 #include "SoundControlBusProxy.h"
 #include "SoundControlBusMixProxy.h"
-#include "SoundModulationGeneratorLFOProxy.h"
+#include "SoundModulationGeneratorProxy.h"
 #include "SoundModulationPatchProxy.h"
 #include "SoundModulationProxy.h"
 #include "UObject/UObjectIterator.h"
 #include "UObject/WeakObjectPtr.h"
 
-
-#if !UE_BUILD_SHIPPING
-#include "AudioModulationDebugger.h"
-#endif // !UE_BUILD_SHIPPING
-
 DECLARE_DWORD_COUNTER_STAT(TEXT("Bus Count"),	STAT_AudioModulationBusCount, STATGROUP_AudioModulation)
-DECLARE_DWORD_COUNTER_STAT(TEXT("LFO Count"),	STAT_AudioModulationLFOCount, STATGROUP_AudioModulation)
+DECLARE_DWORD_COUNTER_STAT(TEXT("Generator Count"),	STAT_AudioModulationGeneratorCount, STATGROUP_AudioModulation)
 DECLARE_DWORD_COUNTER_STAT(TEXT("Mix Count"),	STAT_AudioModulationMixCount, STATGROUP_AudioModulation)
 DECLARE_DWORD_COUNTER_STAT(TEXT("Patch Count"), STAT_AudioModulationPatchCount, STATGROUP_AudioModulation)
 DECLARE_DWORD_COUNTER_STAT(TEXT("Render Queue Commands Processed"), STAT_AudioModulationProcQueueCount, STATGROUP_AudioModulation)
@@ -38,27 +39,11 @@ namespace AudioModulation
 	{
 		Patch,
 		Bus,
-		LFO,
+		Generator,
 
 		COUNT
 	};
 
-	struct FProfileStageInfo
-	{
-		USoundControlBus* Bus;
-		FSoundModulationMixValue Value;
-
-		FProfileStageInfo(const FModulatorBusMixStageProxy& InProxy)
-			: Bus(nullptr)
-			, Value(InProxy.Value)
-		{
-		}
-	};
-
-	FAudioModulationSystem::FAudioModulationSystem()
-	{
-	}
-
 	void FAudioModulationSystem::Initialize(const FAudioPluginInitializationParams& InitializationParams)
 	{
 		// Load all parameters listed in settings and keep loaded.  Must be done on initialization as
@@ -107,7 +92,7 @@ namespace AudioModulation
 
 	void FAudioModulationSystem::ActivateBus(const USoundControlBus& InBus)
 	{
-		RunCommandOnProcessingThread([this, Settings = FControlBusSettings(InBus)]()
+		RunCommandOnProcessingThread([this, Settings = FControlBusSettings(InBus, AudioDeviceId)]()
 		{
 			FBusHandle BusHandle = FBusHandle::Create(Settings, RefProxies.Buses, *this);
 			ManuallyActivatedBuses.Add(MoveTemp(BusHandle));
@@ -134,15 +119,15 @@ namespace AudioModulation
 
 	void FAudioModulationSystem::ActivateBusMix(const USoundControlBusMix& InBusMix)
 	{
-		ActivateBusMix(FModulatorBusMixSettings(InBusMix));
+		ActivateBusMix(FModulatorBusMixSettings(InBusMix, AudioDeviceId));
 	}
 
-	void FAudioModulationSystem::ActivateLFO(const USoundModulationGeneratorLFO& InLFO)
+	void FAudioModulationSystem::ActivateGenerator(const USoundModulationGenerator& InGenerator)
 	{
-		RunCommandOnProcessingThread([this, Settings = FModulatorLFOSettings(InLFO)]()
+		RunCommandOnProcessingThread([this, Settings = FModulationGeneratorSettings(InGenerator, AudioDeviceId)]()
 		{
-			FLFOHandle LFOHandle = FLFOHandle::Create(Settings, RefProxies.LFOs, *this);
-			ManuallyActivatedLFOs.Add(MoveTemp(LFOHandle));
+			FGeneratorHandle GeneratorHandle = FGeneratorHandle::Create(Settings, RefProxies.Generators, *this);
+			ManuallyActivatedGenerators.Add(MoveTemp(GeneratorHandle));
 		});
 	}
 
@@ -196,14 +181,14 @@ namespace AudioModulation
 		});
 	}
 
-	void FAudioModulationSystem::DeactivateLFO(const USoundModulationGeneratorLFO& InLFO)
+	void FAudioModulationSystem::DeactivateGenerator(const USoundModulationGenerator& InGenerator)
 	{
-		RunCommandOnProcessingThread([this, LFOId = static_cast<FLFOId>(InLFO.GetUniqueID())]()
+		RunCommandOnProcessingThread([this, GeneratorId = static_cast<FGeneratorId>(InGenerator.GetUniqueID())]()
 		{
-			FLFOHandle LFOHandle = FLFOHandle::Get(LFOId, RefProxies.LFOs);
-			if (LFOHandle.IsValid())
+			FGeneratorHandle GeneratorHandle = FGeneratorHandle::Get(GeneratorId, RefProxies.Generators);
+			if (GeneratorHandle.IsValid())
 			{
-				ManuallyActivatedLFOs.Remove(LFOHandle);
+				ManuallyActivatedGenerators.Remove(GeneratorHandle);
 			}
 		});
 	}
@@ -242,13 +227,13 @@ namespace AudioModulation
 			}
 			break;
 
-			case EModulatorType::LFO:
+			case EModulatorType::Generator:
 			{
-				if (const FModulatorLFOProxy* LFOProxy = RefProxies.LFOs.Find(static_cast<FLFOId>(InModulatorHandle.GetModulatorId())))
+				if (const FModulatorGeneratorProxy* GeneratorProxy = RefProxies.Generators.Find(static_cast<FGeneratorId>(InModulatorHandle.GetModulatorId())))
 				{
-					if (!LFOProxy->IsBypassed())
+					if (!GeneratorProxy->IsBypassed())
 					{
-						OutValue = LFOProxy->GetValue();
+						OutValue = GeneratorProxy->GetValue();
 						return true;
 					}
 				}
@@ -265,6 +250,11 @@ namespace AudioModulation
 		return false;
 	}
 
+	Audio::FDeviceId FAudioModulationSystem::GetAudioDeviceId() const
+	{
+		return AudioDeviceId;
+	}
+
 	Audio::FModulationParameter FAudioModulationSystem::GetParameter(FName InParamName) const
 	{
 		Audio::FModulationParameter Parameter;
@@ -317,17 +307,18 @@ namespace AudioModulation
 			++CommandsProcessed;
 		}
 
-		// Update LFOs (prior to bus mixing to avoid single-frame latency)
-		for (TPair<FLFOId, FModulatorLFOProxy>& Pair : RefProxies.LFOs)
+		// Update Generators (prior to bus mixing to avoid single-frame latency)
+		for (TPair<FGeneratorId, FModulatorGeneratorProxy>& Pair : RefProxies.Generators)
 		{
+			Pair.Value.PumpCommands();
 			Pair.Value.Update(InElapsed);
 		}
 
-		// Reset buses & refresh cached LFO
+		// Reset buses & refresh cached Generator
 		for (TPair<FBusId, FControlBusProxy>& Pair : RefProxies.Buses)
 		{
 			Pair.Value.Reset();
-			Pair.Value.MixLFO();
+			Pair.Value.MixGenerators();
 		}
 
 		// Update mix values and apply to prescribed buses.
@@ -392,7 +383,7 @@ namespace AudioModulation
 		// Log stats
 		SET_DWORD_STAT(STAT_AudioModulationBusCount, RefProxies.Buses.Num());
 		SET_DWORD_STAT(STAT_AudioModulationMixCount, RefProxies.BusMixes.Num());
-		SET_DWORD_STAT(STAT_AudioModulationLFOCount, RefProxies.LFOs.Num());
+		SET_DWORD_STAT(STAT_AudioModulationGeneratorCount, RefProxies.Generators.Num());
 		SET_DWORD_STAT(STAT_AudioModulationPatchCount, RefProxies.Patches.Num());
 		SET_DWORD_STAT(STAT_AudioModulationProcQueueCount, CommandsProcessed);
 
@@ -504,22 +495,22 @@ namespace AudioModulation
 			return INDEX_NONE;
 		}
 
-		if (RegisterModulator<FPatchHandle, USoundModulationPatch, FModulationPatchSettings, FPatchProxyMap>(InHandleId, InModulatorBase, RefProxies.Patches, RefModulators.PatchMap, OutParameter))
+		if (RegisterModulator<FPatchHandle, USoundModulationPatch, FModulationPatchSettings, FPatchProxyMap>(InHandleId, InModulatorBase, RefProxies.Patches, RefModulators.PatchMap))
 		{
 			return static_cast<Audio::FModulatorTypeId>(EModulatorType::Patch);
 		}
 			
-		if (RegisterModulator<FBusHandle, USoundControlBus, FControlBusSettings, FBusProxyMap>(InHandleId, InModulatorBase, RefProxies.Buses, RefModulators.BusMap, OutParameter))
+		if (RegisterModulator<FBusHandle, USoundControlBus, FControlBusSettings, FBusProxyMap>(InHandleId, InModulatorBase, RefProxies.Buses, RefModulators.BusMap))
 		{
 			return static_cast<Audio::FModulatorTypeId>(EModulatorType::Bus);
 		}
 
-		if (RegisterModulator<FLFOHandle, USoundModulationGeneratorLFO, FModulatorLFOSettings, FLFOProxyMap>(InHandleId, InModulatorBase, RefProxies.LFOs, RefModulators.LFOMap, OutParameter))
+		if (RegisterModulator<FGeneratorHandle, USoundModulationGenerator, FModulationGeneratorSettings, FGeneratorProxyMap>(InHandleId, InModulatorBase, RefProxies.Generators, RefModulators.GeneratorMap))
 		{
-			return static_cast<Audio::FModulatorTypeId>(EModulatorType::LFO);
+			return static_cast<Audio::FModulatorTypeId>(EModulatorType::Generator);
 		}
 
-		UE_LOG(LogAudioModulation, Warning, TEXT("Modulator type  of '%s' unsupported by generic control modulation."), *InModulatorBase->GetName());
+		ensureMsgf(false, TEXT("Modulator type of '%s' failed to register: Unsupported type."), *InModulatorBase->GetName());
 		return INDEX_NONE;
 	}
 
@@ -545,10 +536,10 @@ namespace AudioModulation
 				}
 			}
 
-			FLFOHandle LFOHandle = FLFOHandle::Get(static_cast<FLFOId>(InModulatorId), RefProxies.LFOs);
-			if (LFOHandle.IsValid())
+			FGeneratorHandle GeneratorHandle = FGeneratorHandle::Get(static_cast<FGeneratorId>(InModulatorId), RefProxies.Generators);
+			if (GeneratorHandle.IsValid())
 			{
-				if (TArray<uint32>* RefObjectIds = RefModulators.LFOMap.Find(LFOHandle))
+				if (TArray<uint32>* RefObjectIds = RefModulators.GeneratorMap.Find(GeneratorHandle))
 				{
 					RefObjectIds->Add(InHandleId);
 				}
@@ -558,7 +549,7 @@ namespace AudioModulation
 
 	void FAudioModulationSystem::SoloBusMix(const USoundControlBusMix& InBusMix)
 	{
-		RunCommandOnProcessingThread([this, BusMixSettings = FModulatorBusMixSettings(InBusMix)]()
+		RunCommandOnProcessingThread([this, BusMixSettings = FModulatorBusMixSettings(InBusMix, AudioDeviceId)]()
 		{
 			bool bMixActive = false;
 			for (TPair<FBusMixId, FModulatorBusMixProxy>& Pair : RefProxies.BusMixes)
@@ -596,8 +587,8 @@ namespace AudioModulation
 				return;
 			}
 
-			FLFOHandle LFOHandle = FLFOHandle::Get(static_cast<FLFOId>(ModId), RefProxies.LFOs);
-			if (UnregisterModulator<FLFOHandle>(LFOHandle, RefModulators.LFOMap, HandleId))
+			FGeneratorHandle GeneratorHandle = FGeneratorHandle::Get(static_cast<FGeneratorId>(ModId), RefProxies.Generators);
+			if (UnregisterModulator<FGeneratorHandle>(GeneratorHandle, RefModulators.GeneratorMap, HandleId))
 			{
 				return;
 			}
@@ -639,7 +630,7 @@ namespace AudioModulation
 		TArray<FModulatorBusMixStageSettings> StageSettings;
 		for (const FSoundControlBusMixStage& Stage : InStages)
 		{
-			StageSettings.Emplace(Stage);
+			StageSettings.Emplace(Stage, AudioDeviceId);
 		}
 	
 		RunCommandOnProcessingThread([this, MixId, StageSettings, InFadeTime]()
@@ -718,7 +709,7 @@ namespace AudioModulation
 
 	void FAudioModulationSystem::UpdateMix(const USoundControlBusMix& InMix, float InFadeTime)
 	{
-		RunCommandOnProcessingThread([this, MixSettings = FModulatorBusMixSettings(InMix), InFadeTime]()
+		RunCommandOnProcessingThread([this, MixSettings = FModulatorBusMixSettings(InMix, AudioDeviceId), InFadeTime]()
 		{
 			FBusMixHandle BusMixHandle = FBusMixHandle::Get(MixSettings.GetId(), RefProxies.BusMixes);
 			if (BusMixHandle.IsValid())
@@ -744,19 +735,19 @@ namespace AudioModulation
 
 	void FAudioModulationSystem::UpdateModulator(const USoundModulatorBase& InModulator)
 	{
-		if (const USoundModulationGeneratorLFO* InLFO = Cast<USoundModulationGeneratorLFO>(&InModulator))
+		if (const USoundModulationGenerator* InGenerator = Cast<USoundModulationGenerator>(&InModulator))
 		{
-			RunCommandOnProcessingThread([this, LFOSettings = FModulatorLFOSettings(*InLFO)]()
+			RunCommandOnProcessingThread([this, GeneratorSettings = FModulationGeneratorSettings(*InGenerator, AudioDeviceId)]()
 			{
-				FLFOHandle LFOHandle = FLFOHandle::Get(LFOSettings.GetId(), RefProxies.LFOs);
-				if (LFOHandle.IsValid())
+				FGeneratorHandle GeneratorHandle = FGeneratorHandle::Get(GeneratorSettings.GetId(), RefProxies.Generators);
+				if (GeneratorHandle.IsValid())
 				{
-					LFOHandle.FindProxy() = LFOSettings;
+					GeneratorHandle.FindProxy() = GeneratorSettings;
 				}
 #if !UE_BUILD_SHIPPING
 				else
 				{
-					UE_LOG(LogAudioModulation, Verbose, TEXT("Update to '%s' Ignored: LFO is inactive."), *LFOSettings.GetName());
+					UE_LOG(LogAudioModulation, Verbose, TEXT("Update to '%s' Ignored: Generator is inactive."), *GeneratorSettings.GetName());
 				}
 #endif // !UE_BUILD_SHIPPING
 			});
@@ -764,7 +755,7 @@ namespace AudioModulation
 
 		if (const USoundControlBus* InBus = Cast<USoundControlBus>(&InModulator))
 		{
-			RunCommandOnProcessingThread([this, BusSettings = FControlBusSettings(*InBus)]()
+			RunCommandOnProcessingThread([this, BusSettings = FControlBusSettings(*InBus, AudioDeviceId)]()
 			{
 				FBusHandle BusHandle = FBusHandle::Get(BusSettings.GetId(), RefProxies.Buses);
 				if (BusHandle.IsValid())
@@ -782,7 +773,7 @@ namespace AudioModulation
 
 		if (const USoundControlBusMix* InMix = Cast<USoundControlBusMix>(&InModulator))
 		{
-			RunCommandOnProcessingThread([this, BusMixSettings = FModulatorBusMixSettings(*InMix)]()
+			RunCommandOnProcessingThread([this, BusMixSettings = FModulatorBusMixSettings(*InMix, AudioDeviceId)]()
 			{
 				FBusMixHandle BusMixHandle = FBusMixHandle::Get(BusMixSettings.GetId(), RefProxies.BusMixes);
 				if (BusMixHandle.IsValid())
@@ -800,7 +791,7 @@ namespace AudioModulation
 
 		if (const USoundModulationPatch* InPatch = Cast<USoundModulationPatch>(&InModulator))
 		{
-			RunCommandOnProcessingThread([this, PatchSettings = FModulationPatchSettings(*InPatch)]()
+			RunCommandOnProcessingThread([this, PatchSettings = FModulationPatchSettings(*InPatch, AudioDeviceId)]()
 			{
 				FPatchHandle PatchHandle = FPatchHandle::Get(PatchSettings.GetId(), RefProxies.Patches);
 				if (PatchHandle.IsValid())
diff --git a/Engine/Plugins/Runtime/AudioModulation/Source/AudioModulation/Private/AudioModulationSystem.h b/Engine/Plugins/Runtime/AudioModulation/Source/AudioModulation/Private/AudioModulationSystem.h
--- a/Engine/Plugins/Runtime/AudioModulation/Source/AudioModulation/Private/AudioModulationSystem.h
+++ b/Engine/Plugins/Runtime/AudioModulation/Source/AudioModulation/Private/AudioModulationSystem.h
@@ -15,8 +15,8 @@
 #include "SoundModulationPatchProxy.h"
 #include "SoundModulationProxy.h"
 #include "SoundModulationValue.h"
-#include "SoundModulationGeneratorLFO.h"
-#include "SoundModulationGeneratorLFOProxy.h"
+#include "SoundModulationGenerator.h"
+#include "SoundModulationGeneratorProxy.h"
 #include "Templates/Function.h"
 
 #if WITH_AUDIOMODULATION
@@ -37,44 +37,44 @@ namespace AudioModulation
 	struct FReferencedProxies
 	{
 		FBusMixProxyMap BusMixes;
-		FBusProxyMap    Buses;
-		FLFOProxyMap    LFOs;
-		FPatchProxyMap	Patches;
+		FBusProxyMap Buses;
+		FGeneratorProxyMap Generators;
+		FPatchProxyMap Patches;
 	};
 
 	struct FReferencedModulators
 	{
 		TMap<FPatchHandle, TArray<uint32>> PatchMap;
 		TMap<FBusHandle, TArray<uint32>> BusMap;
-		TMap<FLFOHandle, TArray<uint32>> LFOMap;
+		TMap<FGeneratorHandle, TArray<uint32>> GeneratorMap;
 	};
 
 	class FAudioModulationSystem
 	{
 	public:
-		FAudioModulationSystem();
-
 		void Initialize(const FAudioPluginInitializationParams& InitializationParams);
 
 		void ActivateBus(const USoundControlBus& InBus);
 		void ActivateBusMix(const FModulatorBusMixSettings& InSettings);
 		void ActivateBusMix(const USoundControlBusMix& InBusMix);
-		void ActivateLFO(const USoundModulationGeneratorLFO& InLFO);
+		void ActivateGenerator(const USoundModulationGenerator& InGenerator);
 
 		/**
-		 * Deactivates respectively typed (i.e. BusMix, Bus, etc.) object proxy if no longer referenced.
+		 * Deactivates respectively typed (i.e. BusMix, Bus, Generator, etc.) object proxy if no longer referenced.
 		 * If still referenced, will wait until references are finished before destroying.
 		 */
 		void DeactivateBus(const USoundControlBus& InBus);
 		void DeactivateBusMix(const USoundControlBusMix& InBusMix);
 		void DeactivateAllBusMixes();
-		void DeactivateLFO(const USoundModulationGeneratorLFO& InLFO);
+		void DeactivateGenerator(const USoundModulationGenerator& InGenerator);
 
 		Audio::FModulationParameter GetParameter(FName InParamName) const;
 
 		void ProcessModulators(const double InElapsed);
 		void SoloBusMix(const USoundControlBusMix& InBusMix);
 
+		Audio::FDeviceId GetAudioDeviceId() const;
+
 		Audio::FModulatorTypeId RegisterModulator(Audio::FModulatorHandleId InHandleId, const USoundModulatorBase* InModulatorBase, Audio::FModulationParameter& OutParameter);
 		void RegisterModulator(Audio::FModulatorHandleId InHandleId, Audio::FModulatorId InModulatorId);
 		bool GetModulatorValue(const Audio::FModulatorHandle& ModulatorHandle, float& OutValue) const;
@@ -118,13 +118,13 @@ namespace AudioModulation
 		void RunCommandOnProcessingThread(TUniqueFunction<void()> Cmd);
 
 		template <typename THandleType, typename TModType, typename TModSettings, typename TMapType>
-		bool RegisterModulator(Audio::FModulatorHandleId InHandleId, const USoundModulatorBase* InModulatorBase, TMapType& ProxyMap, TMap<THandleType, TArray<uint32>>& ModMap, Audio::FModulationParameter& OutParameter)
+		bool RegisterModulator(Audio::FModulatorHandleId InHandleId, const USoundModulatorBase* InModulatorBase, TMapType& ProxyMap, TMap<THandleType, TArray<uint32>>& ModMap)
 		{
 			check(InHandleId != INDEX_NONE);
 
 			if (const TModType* Mod = Cast<TModType>(InModulatorBase))
 			{
-				RunCommandOnProcessingThread([this, Modulator = TModSettings(*Mod), InHandleId, PassedProxyMap = &ProxyMap, PassedModMap = &ModMap]()
+				RunCommandOnProcessingThread([this, Modulator = TModSettings(*Mod, AudioDeviceId), InHandleId, PassedProxyMap = &ProxyMap, PassedModMap = &ModMap]()
 				{
 					check(PassedProxyMap);
 					check(PassedModMap);
@@ -169,9 +169,9 @@ namespace AudioModulation
 
 		FReferencedProxies RefProxies;
 
-		TSet<FBusHandle>    ManuallyActivatedBuses;
+		TSet<FBusHandle> ManuallyActivatedBuses;
 		TSet<FBusMixHandle> ManuallyActivatedBusMixes;
-		TSet<FLFOHandle>    ManuallyActivatedLFOs;
+		TSet<FGeneratorHandle> ManuallyActivatedGenerators;
 
 		// Command queue to be consumed on processing thread 
 		TQueue<TUniqueFunction<void()>, EQueueMode::Mpsc> ProcessingThreadCommandQueue;
@@ -182,6 +182,8 @@ namespace AudioModulation
 		// Collection of maps with modulator handles to referencing object ids used by externally managing objects
 		FReferencedModulators RefModulators;
 
+		Audio::FDeviceId AudioDeviceId = INDEX_NONE;
+
 #if !UE_BUILD_SHIPPING
 	public:
 		bool OnPostHelp(FCommonViewportClient* ViewportClient, const TCHAR* Stream);
@@ -198,22 +200,6 @@ namespace AudioModulation
 		friend FModulationPatchRefProxy;
 		friend FModulatorBusMixStageProxy;
 	};
-
-	static void IterateModSystems(TUniqueFunction<void(FAudioModulationSystem&)> InFunction)
-	{
-		if (FAudioDeviceManager* DeviceManager = FAudioDeviceManager::Get())
-		{
-			TArray<FAudioDevice*> Devices = DeviceManager->GetAudioDevices();
-			DeviceManager->IterateOverAllDevices([ModFunction = MoveTemp(InFunction)](Audio::FDeviceId DeviceId, FAudioDevice* AudioDevice)
-			{
-				if (AudioDevice && AudioDevice->IsModulationPluginEnabled() && AudioDevice->ModulationInterface.IsValid())
-				{
-					auto ModulationInterface = static_cast<AudioModulation::FAudioModulation*>(AudioDevice->ModulationInterface.Get());
-					ModFunction(*ModulationInterface->GetModulationSystem());
-				}
-			});
-		}
-	}
 } // namespace AudioModulation
 
 #else // WITH_AUDIOMODULATION
@@ -242,12 +228,12 @@ namespace AudioModulation
 		void ActivateBus(const USoundControlBus& InBus) { }
 		void ActivateBusMix(const FModulatorBusMixSettings& InSettings) { }
 		void ActivateBusMix(const USoundControlBusMix& InBusMix) { }
-		void ActivateLFO(const USoundModulationGeneratorLFO& InLFO) { }
+		void ActivateGenerator(const USoundModulationGenerator& InGenerator) { }
 
 		void DeactivateAllBusMixes() { }
 		void DeactivateBus(const USoundControlBus& InBus) { }
 		void DeactivateBusMix(const USoundControlBus& InBusMix) { }
-		void DeactivateLFO(const USoundModulationGeneratorLFO& InLFO) { }
+		void DeactivateGenerator(const USoundModulationGenerator& InGenerator) { }
 
 		void SaveMixToProfile(const USoundControlBusMix& InBusMix, const int32 InProfileIndex) { }
 		void LoadMixFromProfile(const int32 InProfileIndex, USoundControlBusMix& OutBusMix) { }
diff --git a/Engine/Plugins/Runtime/AudioModulation/Source/AudioModulation/Private/Generators/SoundModulationLFO.cpp b/Engine/Plugins/Runtime/AudioModulation/Source/AudioModulation/Private/Generators/SoundModulationLFO.cpp
new file mode 100644
--- /dev/null
+++ b/Engine/Plugins/Runtime/AudioModulation/Source/AudioModulation/Private/Generators/SoundModulationLFO.cpp
@@ -0,0 +1,72 @@
+// Copyright Epic Games, Inc. All Rights Reserved.
+#include "Generators/SoundModulationLFO.h"
+
+#include "AudioDevice.h"
+#include "AudioModulation.h"
+#include "AudioModulationLogging.h"
+#include "AudioModulationSystem.h"
+#include "Engine/World.h"
+
+
+namespace AudioModulation
+{
+#if !UE_BUILD_SHIPPING
+	const FString FLFOGenerator::DebugName = TEXT("LFO");
+
+	void FLFOGenerator::GetDebugValues(TArray<FString>& OutDebugValues) const
+	{
+		OutDebugValues.Add(FString::Printf(TEXT("%.4f"), GetValue()));
+		OutDebugValues.Add(FString::Printf(TEXT("%.4f"), LFO.GetGain()));
+		OutDebugValues.Add(FString::Printf(TEXT("%.4f"), LFO.GetFrequency()));
+		OutDebugValues.Add(FString::Printf(TEXT("%.4f"), Params.Offset));
+
+		switch (LFO.GetType())
+		{
+			case Audio::ELFO::DownSaw:
+				OutDebugValues.Add(TEXT("DownSaw"));
+				break;
+
+			case Audio::ELFO::Exponential:
+				OutDebugValues.Add(TEXT("Exponential"));
+				break;
+
+			case Audio::ELFO::RandomSampleHold:
+				OutDebugValues.Add(TEXT("Random (Sample & Hold)"));
+				break;
+
+			case Audio::ELFO::Sine:
+				OutDebugValues.Add(TEXT("Sine"));
+				break;
+
+			case Audio::ELFO::Square:
+				OutDebugValues.Add(TEXT("Square"));
+				break;
+
+			case Audio::ELFO::Triangle:
+				OutDebugValues.Add(TEXT("Triangle"));
+				break;
+
+			case Audio::ELFO::UpSaw:
+				OutDebugValues.Add(TEXT("Up Saw"));
+				break;
+
+			default:
+				static_assert(static_cast<int32>(Audio::ELFO::NumLFOTypes) == 7, "Missing LFO type case coverage");
+				break;
+		}
+	}
+
+	const FString& FLFOGenerator::GetDebugName() const
+	{
+		return DebugName;
+	}
+#endif // !UE_BUILD_SHIPPING
+} // namespace AudioModulation
+
+#if !UE_BUILD_SHIPPING
+const FString& USoundModulationGeneratorLFO::GetDebugName()
+{
+	using namespace AudioModulation;
+	return FLFOGenerator::DebugName;
+}
+#endif // !UE_BUILD_SHIPPING
\ No newline at end of file
diff --git a/Engine/Plugins/Runtime/AudioModulation/Source/AudioModulation/Private/SoundControlBus.cpp b/Engine/Plugins/Runtime/AudioModulation/Source/AudioModulation/Private/SoundControlBus.cpp
--- a/Engine/Plugins/Runtime/AudioModulation/Source/AudioModulation/Private/SoundControlBus.cpp
+++ b/Engine/Plugins/Runtime/AudioModulation/Source/AudioModulation/Private/SoundControlBus.cpp
@@ -9,7 +9,7 @@
 #include "DSP/BufferVectorOperations.h"
 #include "Engine/World.h"
 #include "SoundControlBusProxy.h"
-#include "SoundModulationGeneratorLFO.h"
+
 
 USoundControlBus::USoundControlBus(const FObjectInitializer& ObjectInitializer)
 	: Super(ObjectInitializer)
@@ -41,7 +41,7 @@ void USoundControlBus::PostEditChangeProperty(FPropertyChangedEvent& InPropertyC
 			Address = GetName();
 		}
 
-		AudioModulation::IterateModSystems([this](AudioModulation::FAudioModulationSystem & OutModSystem)
+		AudioModulation::IterateModulationImpl([this](AudioModulation::FAudioModulation& OutModSystem)
 		{
 			OutModSystem.UpdateModulator(*this);
 		});
@@ -71,6 +71,8 @@ void USoundControlBus::PostRename(UObject* OldOuter, const FName OldName)
 
 void USoundControlBus::BeginDestroy()
 {
+	using namespace AudioModulation;
+
 	Super::BeginDestroy();
 
 	if (UWorld* World = GetWorld())
@@ -80,9 +82,9 @@ void USoundControlBus::BeginDestroy()
 			check(AudioDevice->IsModulationPluginEnabled());
 			if (IAudioModulation* ModulationInterface = AudioDevice->ModulationInterface.Get())
 			{
-				auto ModSystem = static_cast<AudioModulation::FAudioModulation*>(ModulationInterface)->GetModulationSystem();
-				check(ModSystem);
-				ModSystem->DeactivateBus(*this);
+				FAudioModulation* Modulation = static_cast<FAudioModulation*>(ModulationInterface);
+				check(Modulation);
+				Modulation->DeactivateBus(*this);
 			}
 		}
 	}
diff --git a/Engine/Plugins/Runtime/AudioModulation/Source/AudioModulation/Private/SoundControlBusMix.cpp b/Engine/Plugins/Runtime/AudioModulation/Source/AudioModulation/Private/SoundControlBusMix.cpp
--- a/Engine/Plugins/Runtime/AudioModulation/Source/AudioModulation/Private/SoundControlBusMix.cpp
+++ b/Engine/Plugins/Runtime/AudioModulation/Source/AudioModulation/Private/SoundControlBusMix.cpp
@@ -7,7 +7,6 @@
 #include "AudioModulationLogging.h"
 #include "AudioModulationProfileSerializer.h"
 #include "AudioModulationStatics.h"
-#include "AudioModulationSystem.h"
 #include "Engine/World.h"
 #include "SoundControlBus.h"
 
@@ -38,6 +37,8 @@ USoundControlBusMix::USoundControlBusMix(const FObjectInitializer& ObjectInitial
 
 void USoundControlBusMix::BeginDestroy()
 {
+	using namespace AudioModulation;
+
 	Super::BeginDestroy();
 
 	if (UWorld* World = GetWorld())
@@ -48,9 +49,9 @@ void USoundControlBusMix::BeginDestroy()
 			{
 				if (IAudioModulation* ModulationInterface = AudioDevice->ModulationInterface.Get())
 				{
-					auto ModSystem = static_cast<AudioModulation::FAudioModulation*>(ModulationInterface)->GetModulationSystem();
-					check(ModSystem);
-					ModSystem->DeactivateBusMix(*this);
+					FAudioModulation* Modulation = static_cast<FAudioModulation*>(ModulationInterface);
+					check(Modulation);
+					Modulation->DeactivateBusMix(*this);
 				}
 			}
 		}
@@ -59,25 +60,25 @@ void USoundControlBusMix::BeginDestroy()
 
 void USoundControlBusMix::ActivateMix()
 {
-	AudioModulation::IterateModSystems([this](AudioModulation::FAudioModulationSystem& OutModSystem)
+	AudioModulation::IterateModulationImpl([this](AudioModulation::FAudioModulation& OutModulation)
 	{
-		OutModSystem.ActivateBusMix(*this);
+		OutModulation.ActivateBusMix(*this);
 	});
 }
 
 void USoundControlBusMix::DeactivateMix()
 {
-	AudioModulation::IterateModSystems([this](AudioModulation::FAudioModulationSystem& OutModSystem)
+	AudioModulation::IterateModulationImpl([this](AudioModulation::FAudioModulation& OutModulation)
 	{
-		OutModSystem.DeactivateBusMix(*this);
+		OutModulation.DeactivateBusMix(*this);
 	});
 }
 
 void USoundControlBusMix::DeactivateAllMixes()
 {
-	AudioModulation::IterateModSystems([this](AudioModulation::FAudioModulationSystem& OutModSystem)
+	AudioModulation::IterateModulationImpl([this](AudioModulation::FAudioModulation& OutModulation)
 	{
-		OutModSystem.DeactivateAllBusMixes();
+		OutModulation.DeactivateAllBusMixes();
 	});
 }
 
@@ -132,9 +133,9 @@ void USoundControlBusMix::OnPropertyChanged(FProperty* Property, EPropertyChange
 		}
 	}
 
-	AudioModulation::IterateModSystems([this](AudioModulation::FAudioModulationSystem& OutModSystem)
+	AudioModulation::IterateModulationImpl([this](AudioModulation::FAudioModulation& OutModulation)
 	{
-		OutModSystem.UpdateMix(*this);
+		OutModulation.UpdateMix(*this);
 	});
 }
 #endif // WITH_EDITOR
@@ -160,9 +161,9 @@ void USoundControlBusMix::SaveMixToProfile()
 
 void USoundControlBusMix::SoloMix()
 {
-	AudioModulation::IterateModSystems([this](AudioModulation::FAudioModulationSystem& OutModSystem)
+	AudioModulation::IterateModulationImpl([this](AudioModulation::FAudioModulation& OutModulation)
 	{
-		OutModSystem.SoloBusMix(*this);
+		OutModulation.SoloBusMix(*this);
 	});
 }
 
diff --git a/Engine/Plugins/Runtime/AudioModulation/Source/AudioModulation/Private/SoundControlBusMixProxy.cpp b/Engine/Plugins/Runtime/AudioModulation/Source/AudioModulation/Private/SoundControlBusMixProxy.cpp
--- a/Engine/Plugins/Runtime/AudioModulation/Source/AudioModulation/Private/SoundControlBusMixProxy.cpp
+++ b/Engine/Plugins/Runtime/AudioModulation/Source/AudioModulation/Private/SoundControlBusMixProxy.cpp
@@ -3,6 +3,7 @@
 
 #include "Audio/AudioAddressPattern.h"
 #include "AudioDevice.h"
+#include "AudioDeviceManager.h"
 #include "AudioModulation.h"
 #include "AudioModulationLogging.h"
 #include "AudioModulationProfileSerializer.h"
@@ -18,13 +19,13 @@ namespace AudioModulation
 {
 	const FBusMixId InvalidBusMixId = INDEX_NONE;
 
-	FModulatorBusMixStageSettings::FModulatorBusMixStageSettings(const FSoundControlBusMixStage& InStage)
+	FModulatorBusMixStageSettings::FModulatorBusMixStageSettings(const FSoundControlBusMixStage& InStage, Audio::FDeviceId InDeviceId)
 		: TModulatorBase<FBusId>(InStage.Bus->GetName(), InStage.Bus->GetUniqueID())
 		, Address(InStage.Bus->Address)
 		, ParamClassId(InStage.Bus->Parameter->GetClass()->GetUniqueID())
 		, ParamId(InStage.Bus->Parameter->GetUniqueID())
 		, Value(InStage.Value)
-		, BusSettings(FControlBusSettings(*InStage.Bus))
+		, BusSettings(FControlBusSettings(*InStage.Bus, InDeviceId))
 	{
 	}
 
@@ -38,14 +39,14 @@ namespace AudioModulation
 	{
 	}
 
-	FModulatorBusMixSettings::FModulatorBusMixSettings(const USoundControlBusMix& InBusMix)
+	FModulatorBusMixSettings::FModulatorBusMixSettings(const USoundControlBusMix& InBusMix, Audio::FDeviceId InDeviceId)
 		: TModulatorBase<FBusMixId>(InBusMix.GetName(), InBusMix.GetUniqueID())
 	{
 		for (const FSoundControlBusMixStage& Stage : InBusMix.MixStages)
 		{
 			if (Stage.Bus)
 			{
-				Stages.Add(FModulatorBusMixStageSettings(Stage));
+				Stages.Add(FModulatorBusMixStageSettings(Stage, InDeviceId));
 			}
 			else
 			{
diff --git a/Engine/Plugins/Runtime/AudioModulation/Source/AudioModulation/Private/SoundControlBusMixProxy.h b/Engine/Plugins/Runtime/AudioModulation/Source/AudioModulation/Private/SoundControlBusMixProxy.h
--- a/Engine/Plugins/Runtime/AudioModulation/Source/AudioModulation/Private/SoundControlBusMixProxy.h
+++ b/Engine/Plugins/Runtime/AudioModulation/Source/AudioModulation/Private/SoundControlBusMixProxy.h
@@ -1,11 +1,11 @@
 // Copyright Epic Games, Inc. All Rights Reserved.
 #pragma once
 
+#include "AudioDeviceManager.h"
 #include "IAudioModulation.h"
 #include "SoundControlBusProxy.h"
 #include "SoundModulationProxy.h"
 #include "SoundModulationValue.h"
-#include "SoundModulationGeneratorLFOProxy.h"
 
 
 // Forward Declarations
@@ -24,7 +24,7 @@ namespace AudioModulation
 	class FModulatorBusMixStageSettings : public TModulatorBase<FBusId>
 	{
 	public:
-		FModulatorBusMixStageSettings(const FSoundControlBusMixStage& InStage);
+		FModulatorBusMixStageSettings(const FSoundControlBusMixStage& InStage, Audio::FDeviceId InDeviceId);
 
 		FString Address;
 		uint32 ParamClassId = INDEX_NONE;
@@ -36,7 +36,7 @@ namespace AudioModulation
 	class FModulatorBusMixSettings : public TModulatorBase<FBusMixId>
 	{
 	public:
-		FModulatorBusMixSettings(const USoundControlBusMix& InBusMix);
+		FModulatorBusMixSettings(const USoundControlBusMix& InBusMix, Audio::FDeviceId InDeviceId);
 
 		TArray<FModulatorBusMixStageSettings> Stages;
 	};
diff --git a/Engine/Plugins/Runtime/AudioModulation/Source/AudioModulation/Private/SoundControlBusProxy.cpp b/Engine/Plugins/Runtime/AudioModulation/Source/AudioModulation/Private/SoundControlBusProxy.cpp
--- a/Engine/Plugins/Runtime/AudioModulation/Source/AudioModulation/Private/SoundControlBusProxy.cpp
+++ b/Engine/Plugins/Runtime/AudioModulation/Source/AudioModulation/Private/SoundControlBusProxy.cpp
@@ -7,7 +7,7 @@
 #include "AudioModulationLogging.h"
 #include "AudioModulationSystem.h"
 #include "Engine/World.h"
-#include "SoundModulationGeneratorLFO.h"
+#include "SoundModulationGenerator.h"
 
 
 namespace AudioModulation
@@ -16,7 +16,7 @@ namespace AudioModulation
 
 	FControlBusProxy::FControlBusProxy()
 		: DefaultValue(0.0f)
-		, LFOValue(1.0f)
+		, GeneratorValue(1.0f)
 		, MixValue(NAN)
 		, bBypass(false)
 	{
@@ -39,14 +39,14 @@ namespace AudioModulation
 		return DefaultValue;
 	}
 
-	const TArray<FLFOHandle>& FControlBusProxy::GetLFOHandles() const
+	const TArray<FGeneratorHandle>& FControlBusProxy::GetGeneratorHandles() const
 	{
-		return LFOHandles;
+		return GeneratorHandles;
 	}
 
-	float FControlBusProxy::GetLFOValue() const
+	float FControlBusProxy::GetGeneratorValue() const
 	{
-		return LFOValue;
+		return GeneratorValue;
 	}
 
 	float FControlBusProxy::GetMixValue() const
@@ -57,29 +57,29 @@ namespace AudioModulation
 	float FControlBusProxy::GetValue() const
 	{
 		const float DefaultMixed = Mix(DefaultValue);
-		return FMath::Clamp(DefaultMixed * LFOValue, 0.0f, 1.0f);
+		return FMath::Clamp(DefaultMixed * GeneratorValue, 0.0f, 1.0f);
 	}
 
 	void FControlBusProxy::Init(const FControlBusSettings& InSettings)
 	{
 		check(ModSystem);
 
-		LFOValue = 1.0f;
+		GeneratorValue = 1.0f;
 		MixValue = NAN;
 		MixFunction = InSettings.MixFunction;
 
 		DefaultValue = FMath::Clamp(InSettings.DefaultValue, 0.0f, 1.0f);
 		bBypass = InSettings.bBypass;
 
-		TArray<FLFOHandle> NewHandles;
-		for (const FModulatorLFOSettings& LFOSettings : InSettings.LFOSettings)
+		TArray<FGeneratorHandle> NewHandles;
+		for (const FModulationGeneratorSettings& GeneratorSettings : InSettings.GeneratorSettings)
 		{
-			NewHandles.Add(FLFOHandle::Create(LFOSettings, ModSystem->RefProxies.LFOs, *ModSystem));
+			NewHandles.Add(FGeneratorHandle::Create(GeneratorSettings, ModSystem->RefProxies.Generators, *ModSystem));
 		}
 
-		// Move vs. reset and adding to original array to avoid potentially clearing handles (and thus current LFO state)
-		// and destroying LFOs if function is called while reinitializing/updating the modulator
-		LFOHandles = MoveTemp(NewHandles);
+		// Move vs. reset and adding to original array to avoid potentially clearing handles (and thus current Generator state)
+		// and destroying generators if function is called while reinitializing/updating the modulator
+		GeneratorHandles = MoveTemp(NewHandles);
 	}
 
 	bool FControlBusProxy::IsBypassed() const
@@ -106,16 +106,16 @@ namespace AudioModulation
 		MixValue = Mix(InValue);
 	}
 
-	void FControlBusProxy::MixLFO()
+	void FControlBusProxy::MixGenerators()
 	{
-		for (const FLFOHandle& Handle: LFOHandles)
+		for (const FGeneratorHandle& Handle: GeneratorHandles)
 		{
 			if (Handle.IsValid())
 			{
-				const FModulatorLFOProxy& LFOProxy = Handle.FindProxy();
-				if (!LFOProxy.IsBypassed())
+				const FModulatorGeneratorProxy& GeneratorProxy = Handle.FindProxy();
+				if (!GeneratorProxy.IsBypassed())
 				{
-					LFOValue *= LFOProxy.GetValue();
+					GeneratorValue *= GeneratorProxy.GetValue();
 				}
 			}
 		}
@@ -123,7 +123,7 @@ namespace AudioModulation
 
 	void FControlBusProxy::Reset()
 	{
-		LFOValue = 1.0f;
+		GeneratorValue = 1.0f;
 		MixValue = NAN;
 	}
 } // namespace AudioModulation
diff --git a/Engine/Plugins/Runtime/AudioModulation/Source/AudioModulation/Private/SoundControlBusProxy.h b/Engine/Plugins/Runtime/AudioModulation/Source/AudioModulation/Private/SoundControlBusProxy.h
--- a/Engine/Plugins/Runtime/AudioModulation/Source/AudioModulation/Private/SoundControlBusProxy.h
+++ b/Engine/Plugins/Runtime/AudioModulation/Source/AudioModulation/Private/SoundControlBusProxy.h
@@ -1,11 +1,12 @@
 // Copyright Epic Games, Inc. All Rights Reserved.
 #pragma once
 
+#include "AudioDeviceManager.h"
 #include "IAudioModulation.h"
 #include "SoundControlBus.h"
 #include "SoundModulationParameter.h"
 #include "SoundModulationProxy.h"
-#include "SoundModulationGeneratorLFOProxy.h"
+#include "SoundModulationGeneratorProxy.h"
 
 
 namespace AudioModulation
@@ -21,21 +22,18 @@ namespace AudioModulation
 		bool bBypass;
 		float DefaultValue;
 
-		TArray<FModulatorLFOSettings> LFOSettings;
+		TArray<FModulationGeneratorSettings> GeneratorSettings;
 		Audio::FModulationMixFunction MixFunction;
 
-		FControlBusSettings(const USoundControlBus& InBus)
+		FControlBusSettings(const USoundControlBus& InBus, Audio::FDeviceId InDeviceId)
 			: TModulatorBase<FBusId>(InBus.GetName(), InBus.GetUniqueID())
 			, bBypass()
 			, DefaultValue(InBus.GetDefaultLinearValue())
 			, MixFunction(InBus.GetMixFunction())
 		{
-			for (const USoundModulationGenerator* Modulator : InBus.Modulators)
+			for (const USoundModulationGenerator* Generator : InBus.Generators)
 			{
-				if (const USoundModulationGeneratorLFO* LFO = Cast<USoundModulationGeneratorLFO>(Modulator))
-				{
-					LFOSettings.Add(*LFO);
-				}
+				GeneratorSettings.Emplace(FModulationGeneratorSettings(*Generator, InDeviceId));
 			}
 		}
 	};
@@ -49,13 +47,13 @@ namespace AudioModulation
 		FControlBusProxy& operator =(const FControlBusSettings& InSettings);
 
 		float GetDefaultValue() const;
-		const TArray<FLFOHandle>& GetLFOHandles() const;
-		float GetLFOValue() const;
+		const TArray<FGeneratorHandle>& GetGeneratorHandles() const;
+		float GetGeneratorValue() const;
 		float GetMixValue() const;
 		float GetValue() const;
 		bool IsBypassed() const;
 		void MixIn(const float InValue);
-		void MixLFO();
+		void MixGenerators();
 		void Reset();
 
 	private:
@@ -65,13 +63,13 @@ namespace AudioModulation
 		float DefaultValue;
 
 		// Cached values
-		float LFOValue;
+		float GeneratorValue;
 		float MixValue;
 
 		bool bBypass;
 
 		Audio::FModulationMixFunction MixFunction;
-		TArray<FLFOHandle> LFOHandles;
+		TArray<FGeneratorHandle> GeneratorHandles;
 	};
 
 	using FBusProxyMap = TMap<FBusId, FControlBusProxy>;
diff --git a/Engine/Plugins/Runtime/AudioModulation/Source/AudioModulation/Private/SoundModulationGenerator.cpp b/Engine/Plugins/Runtime/AudioModulation/Source/AudioModulation/Private/SoundModulationGenerator.cpp
--- a/Engine/Plugins/Runtime/AudioModulation/Source/AudioModulation/Private/SoundModulationGenerator.cpp
+++ b/Engine/Plugins/Runtime/AudioModulation/Source/AudioModulation/Private/SoundModulationGenerator.cpp
@@ -6,21 +6,59 @@
 #include "AudioModulationLogging.h"
 #include "AudioModulationSystem.h"
 #include "Engine/World.h"
+#include "Templates/Function.h"
 
 
-USoundModulationGenerator::USoundModulationGenerator(const FObjectInitializer& ObjectInitializer)
-	: Super(ObjectInitializer)
+namespace AudioModulation
 {
-}
+	void IGenerator::AudioRenderThreadCommand(TUniqueFunction<void()>&& InCommand)
+	{
+		CommandQueue.Enqueue(MoveTemp(InCommand));
+	}
+
+	void IGenerator::PumpCommands()
+	{
+		TUniqueFunction<void()> Cmd;
+		while (!CommandQueue.IsEmpty())
+		{
+			CommandQueue.Dequeue(Cmd);
+			Cmd();
+		}
+	}
+} // namespace AudioModulation
 
 #if WITH_EDITOR
 void USoundModulationGenerator::PostEditChangeProperty(FPropertyChangedEvent& InPropertyChangedEvent)
 {
-	AudioModulation::IterateModSystems([this](AudioModulation::FAudioModulationSystem& OutModSystem)
+	AudioModulation::IterateModulationImpl([this](AudioModulation::FAudioModulation& OutModulation)
 	{
-		OutModSystem.UpdateModulator(*this);
+		OutModulation.UpdateModulator(*this);
 	});
 
 	Super::PostEditChangeProperty(InPropertyChangedEvent);
 }
 #endif // WITH_EDITOR
+
+void USoundModulationGenerator::BeginDestroy()
+{
+	using namespace AudioModulation;
+
+	Super::BeginDestroy();
+
+	UWorld* World = GetWorld();
+	if (!World)
+	{
+		return;
+	}
+
+	if (FAudioDeviceHandle AudioDevice = World->GetAudioDevice())
+	{
+		check(AudioDevice->IsModulationPluginEnabled());
+		if (IAudioModulation* ModulationInterface = AudioDevice->ModulationInterface.Get())
+		{
+			FAudioModulation* Modulation = static_cast<FAudioModulation*>(ModulationInterface);
+			check(Modulation);
+			Modulation->DeactivateGenerator(*this);
+		}
+	}
+}
diff --git a/Engine/Plugins/Runtime/AudioModulation/Source/AudioModulation/Private/SoundModulationGeneratorLFO.cpp b/Engine/Plugins/Runtime/AudioModulation/Source/AudioModulation/Private/SoundModulationGeneratorLFO.cpp
deleted file mode 100644
diff --git a/Engine/Plugins/Runtime/AudioModulation/Source/AudioModulation/Private/SoundModulationGeneratorLFOProxy.cpp b/Engine/Plugins/Runtime/AudioModulation/Source/AudioModulation/Private/SoundModulationGeneratorLFOProxy.cpp
deleted file mode 100644
diff --git a/Engine/Plugins/Runtime/AudioModulation/Source/AudioModulation/Private/SoundModulationGeneratorLFOProxy.h b/Engine/Plugins/Runtime/AudioModulation/Source/AudioModulation/Private/SoundModulationGeneratorLFOProxy.h
deleted file mode 100644
diff --git a/Engine/Plugins/Runtime/AudioModulation/Source/AudioModulation/Private/SoundModulationGeneratorProxy.cpp b/Engine/Plugins/Runtime/AudioModulation/Source/AudioModulation/Private/SoundModulationGeneratorProxy.cpp
new file mode 100644
--- /dev/null
+++ b/Engine/Plugins/Runtime/AudioModulation/Source/AudioModulation/Private/SoundModulationGeneratorProxy.cpp
@@ -0,0 +1,24 @@
+// Copyright Epic Games, Inc. All Rights Reserved.
+#include "SoundModulationGeneratorProxy.h"
+
+#include "AudioModulation.h"
+#include "AudioModulationSystem.h"
+
+
+namespace AudioModulation
+{
+	const FGeneratorId InvalidGeneratorId = INDEX_NONE;
+
+	FModulatorGeneratorProxy::FModulatorGeneratorProxy(const FModulationGeneratorSettings& InSettings, FAudioModulationSystem& InModSystem)
+		: TModulatorProxyRefType(InSettings.GetName(), InSettings.GetId(), InModSystem)
+		, Generator(InSettings.Generator)
+	{
+	}
+
+	FModulatorGeneratorProxy& FModulatorGeneratorProxy::operator =(const FModulationGeneratorSettings& InSettings)
+	{
+		Generator = InSettings.Generator;
+		return *this;
+	}
+
+} // namespace AudioModulation
diff --git a/Engine/Plugins/Runtime/AudioModulation/Source/AudioModulation/Private/SoundModulationGeneratorProxy.h b/Engine/Plugins/Runtime/AudioModulation/Source/AudioModulation/Private/SoundModulationGeneratorProxy.h
new file mode 100644
--- /dev/null
+++ b/Engine/Plugins/Runtime/AudioModulation/Source/AudioModulation/Private/SoundModulationGeneratorProxy.h
@@ -0,0 +1,89 @@
+// Copyright Epic Games, Inc. All Rights Reserved.
+#pragma once
+
+#include "AudioDeviceManager.h"
+#include "SoundModulationGenerator.h"
+#include "SoundModulationProxy.h"
+
+
+namespace AudioModulation
+{
+	// Forward Declarations
+	class FAudioModulationSystem;
+	class FModulatorGeneratorProxy;
+
+	struct FModulationGeneratorSettings;
+
+	// Modulator Ids
+	using FGeneratorId = uint32;
+	extern const FGeneratorId InvalidGeneratorId;
+
+	using FGeneratorProxyMap = TMap<FGeneratorId, FModulatorGeneratorProxy>;
+	using FGeneratorHandle = TProxyHandle<FGeneratorId, FModulatorGeneratorProxy, FModulationGeneratorSettings>;
+
+	struct FModulationGeneratorSettings : public TModulatorBase<FGeneratorId>
+	{
+		FGeneratorPtr Generator;
+
+		FModulationGeneratorSettings(const USoundModulationGenerator& InGenerator, Audio::FDeviceId AudioDeviceId)
+			: TModulatorBase<FGeneratorId>(InGenerator.GetName(), InGenerator.GetUniqueID())
+			, Generator(InGenerator.CreateInstance(AudioDeviceId))
+		{
+		}
+	};
+
+	class FModulatorGeneratorProxy : public TModulatorProxyRefType<FGeneratorId, FModulatorGeneratorProxy, FModulationGeneratorSettings>
+	{
+		FGeneratorPtr Generator;
+
+	public:
+		FModulatorGeneratorProxy() = default;
+		FModulatorGeneratorProxy(const FModulationGeneratorSettings& InSettings, FAudioModulationSystem& InModSystem);
+		
+		FModulatorGeneratorProxy& operator =(const FModulationGeneratorSettings& InSettings);
+
+		float GetValue() const
+		{
+			return Generator->GetValue();
+		}
+
+		bool IsBypassed() const
+		{
+			return Generator->IsBypassed();
+		}
+
+		void Update(double InElapsed)
+		{
+			Generator->Update(InElapsed);
+		}
+
+		void PumpCommands()
+		{
+			Generator->PumpCommands();
+		}
+
+#if !UE_BUILD_SHIPPING
+		TArray<FString> GetDebugValues() const
+		{
+			TArray<FString> DebugValues;
+			DebugValues.Add(GetName());
+			DebugValues.Add(FString::FormatAsNumber(GetRefCount()));
+
+			Generator->GetDebugValues(DebugValues);
+
+			return DebugValues;
+		}
+
+		const FString& GetDebugName() const
+		{
+			return Generator->GetDebugName();
+		}
+#endif // !UE_BUILD_SHIPPING
+
+	private:
+		void Init(const FModulationGeneratorSettings& InSettings)
+		{
+			Generator = InSettings.Generator;
+		}
+	};
+} // namespace AudioModulation
\ No newline at end of file
diff --git a/Engine/Plugins/Runtime/AudioModulation/Source/AudioModulation/Private/SoundModulationPatch.cpp b/Engine/Plugins/Runtime/AudioModulation/Source/AudioModulation/Private/SoundModulationPatch.cpp
--- a/Engine/Plugins/Runtime/AudioModulation/Source/AudioModulation/Private/SoundModulationPatch.cpp
+++ b/Engine/Plugins/Runtime/AudioModulation/Source/AudioModulation/Private/SoundModulationPatch.cpp
@@ -35,9 +35,9 @@ USoundModulationPatch::USoundModulationPatch(const FObjectInitializer& ObjectIni
 #if WITH_EDITOR
 void USoundModulationPatch::PostEditChangeProperty(FPropertyChangedEvent& InPropertyChangedEvent)
 {
-	AudioModulation::IterateModSystems([this](AudioModulation::FAudioModulationSystem& OutModSystem)
+	AudioModulation::IterateModulationImpl([this](AudioModulation::FAudioModulation& OutModulation)
 	{
-		OutModSystem.UpdateModulator(*this);
+		OutModulation.UpdateModulator(*this);
 	});
 
 	Super::PostEditChangeProperty(InPropertyChangedEvent);
@@ -45,9 +45,9 @@ void USoundModulationPatch::PostEditChangeProperty(FPropertyChangedEvent& InProp
 
 void USoundModulationPatch::PostEditChangeChainProperty(FPropertyChangedChainEvent& InPropertyChangedEvent)
 {
-	AudioModulation::IterateModSystems([this](AudioModulation::FAudioModulationSystem& OutModSystem)
+	AudioModulation::IterateModulationImpl([this](AudioModulation::FAudioModulation& OutModulation)
 	{
-		OutModSystem.UpdateModulator(*this);
+		OutModulation.UpdateModulator(*this);
 	});
 
 	Super::PostEditChangeChainProperty(InPropertyChangedEvent);
diff --git a/Engine/Plugins/Runtime/AudioModulation/Source/AudioModulation/Private/SoundModulationPatchProxy.h b/Engine/Plugins/Runtime/AudioModulation/Source/AudioModulation/Private/SoundModulationPatchProxy.h
--- a/Engine/Plugins/Runtime/AudioModulation/Source/AudioModulation/Private/SoundModulationPatchProxy.h
+++ b/Engine/Plugins/Runtime/AudioModulation/Source/AudioModulation/Private/SoundModulationPatchProxy.h
@@ -1,11 +1,11 @@
 // Copyright Epic Games, Inc. All Rights Reserved.
 #pragma once
 
+#include "AudioDeviceManager.h"
 #include "IAudioModulation.h"
 #include "SoundControlBusProxy.h"
 #include "SoundModulationParameter.h"
 #include "SoundModulationProxy.h"
-#include "SoundModulationGeneratorLFOProxy.h"
 #include "Templates/Function.h"
 
 
@@ -23,8 +23,8 @@ namespace AudioModulation
 		const FSoundModulationTransform Transform;
 		const uint8 bSampleAndHold : 1;
 
-		FModulationInputSettings(const FSoundControlModulationInput& InInput)
-			: BusSettings(FControlBusSettings(InInput.GetBusChecked()))
+		FModulationInputSettings(const FSoundControlModulationInput& InInput, Audio::FDeviceId InDeviceId)
+			: BusSettings(InInput.GetBusChecked(), InDeviceId)
 			, Transform(InInput.Transform)
 			, bSampleAndHold(InInput.bSampleAndHold)
 		{
@@ -75,7 +75,7 @@ namespace AudioModulation
 
 		FModulationPatchSettings() = default;
 
-		FModulationPatchSettings(const FSoundControlModulationPatch& InPatch)
+		FModulationPatchSettings(const FSoundControlModulationPatch& InPatch, Audio::FDeviceId InDeviceId)
 			: bBypass(InPatch.bBypass)
 		{
 			if (InPatch.OutputParameter)
@@ -88,12 +88,12 @@ namespace AudioModulation
 			{
 				if (Input.GetBus())
 				{
-					InputSettings.Emplace(Input);
+					InputSettings.Emplace(Input, InDeviceId);
 				}
 			}
 		}
 
-		FModulationPatchSettings(const USoundModulationPatch& InPatch)
+		FModulationPatchSettings(const USoundModulationPatch& InPatch, Audio::FDeviceId InDeviceId)
 			: TModulatorBase<FPatchId>(InPatch.GetName(), InPatch.GetUniqueID())
 			, bBypass(InPatch.PatchSettings.bBypass)
 		{
@@ -106,7 +106,7 @@ namespace AudioModulation
 			{
 				if (Input.GetBus())
 				{
-					InputSettings.Emplace(Input);
+					InputSettings.Emplace(Input, InDeviceId);
 				}
 			}
 		}
diff --git a/Engine/Plugins/Runtime/AudioModulation/Source/AudioModulation/Private/SoundModulationProxy.h b/Engine/Plugins/Runtime/AudioModulation/Source/AudioModulation/Private/SoundModulationProxy.h
--- a/Engine/Plugins/Runtime/AudioModulation/Source/AudioModulation/Private/SoundModulationProxy.h
+++ b/Engine/Plugins/Runtime/AudioModulation/Source/AudioModulation/Private/SoundModulationProxy.h
@@ -14,7 +14,7 @@ namespace AudioModulation
 	/*
 	 * Handle for all ref-counted proxy types, to be used only on the audio thread (not thread safe). 
 	 */
-	template<typename IdType, typename ProxyType, typename ProxiedObject>
+	template<typename IdType, typename ProxyType, typename ProxySettings>
 	class TProxyHandle
 	{
 		IdType Id;
@@ -47,7 +47,7 @@ namespace AudioModulation
 			return ProxyMap->FindChecked(Id);
 		}
 
-		friend FORCEINLINE uint32 GetTypeHash(const TProxyHandle<IdType, ProxyType, ProxiedObject>& InHandle)
+		friend FORCEINLINE uint32 GetTypeHash(const TProxyHandle<IdType, ProxyType, ProxySettings>& InHandle)
 		{
 			return static_cast<uint32>(InHandle.Id);
 		}
@@ -112,56 +112,56 @@ namespace AudioModulation
 		/*
 		 * Creates a handle to a proxy modulation object tracked in the provided InProxyMap if it exists, otherwise returns invalid handle.
 		 */
-		static TProxyHandle<IdType, ProxyType, ProxiedObject> Get(const IdType ObjectId, TMap<IdType, ProxyType>& OutProxyMap)
+		static TProxyHandle<IdType, ProxyType, ProxySettings> Get(const IdType ObjectId, TMap<IdType, ProxyType>& OutProxyMap)
 		{
 			if (ProxyType* Proxy = OutProxyMap.Find(ObjectId))
 			{
 				check(Proxy->ModSystem);
-				return TProxyHandle<IdType, ProxyType, ProxiedObject>(ObjectId, OutProxyMap);
+				return TProxyHandle<IdType, ProxyType, ProxySettings>(ObjectId, OutProxyMap);
 			} 
 
-			return TProxyHandle<IdType, ProxyType, ProxiedObject>();
+			return TProxyHandle<IdType, ProxyType, ProxySettings>();
 		}
 
 		/*
 		 * Creates a handle to a proxy modulation object tracked in the provided InProxyMap.  Creates new proxy if it doesn't exist.
 		 */
-		static TProxyHandle<IdType, ProxyType, ProxiedObject> Create(const ProxiedObject& InObject, TMap<IdType, ProxyType>& OutProxyMap, FAudioModulationSystem& OutModSystem)
+		static TProxyHandle<IdType, ProxyType, ProxySettings> Create(const ProxySettings& InSettings, TMap<IdType, ProxyType>& OutProxyMap, FAudioModulationSystem& OutModSystem)
 		{
-			const IdType ObjectId = static_cast<IdType>(InObject.GetId());
+			const IdType ObjectId = static_cast<IdType>(InSettings.GetId());
 
-			TProxyHandle<IdType, ProxyType, ProxiedObject> NewHandle = Get(InObject.GetId(), OutProxyMap);
+			TProxyHandle<IdType, ProxyType, ProxySettings> NewHandle = Get(InSettings.GetId(), OutProxyMap);
 			if (!NewHandle.IsValid())
 			{
-				UE_LOG(LogAudioModulation, Verbose, TEXT("Proxy created: Id '%u' for object '%s'."), NewHandle.Id, *InObject.GetName());
-				OutProxyMap.Add(ObjectId, ProxyType(InObject, OutModSystem));
-				NewHandle = TProxyHandle<IdType, ProxyType, ProxiedObject>(ObjectId, OutProxyMap);
+				UE_LOG(LogAudioModulation, Verbose, TEXT("Proxy created: Id '%u' for object '%s'."), NewHandle.Id, *InSettings.GetName());
+				OutProxyMap.Add(ObjectId, ProxyType(InSettings, OutModSystem));
+				NewHandle = TProxyHandle<IdType, ProxyType, ProxySettings>(ObjectId, OutProxyMap);
 			}
 
 			return NewHandle;
 		}
 
-		static TProxyHandle<IdType, ProxyType, ProxiedObject> Create(const ProxiedObject& InObject, TMap<IdType, ProxyType>& InProxyMap, FAudioModulationSystem& OutModSystem, TFunction<void(ProxyType&)> OnCreateProxy)
+		static TProxyHandle<IdType, ProxyType, ProxySettings> Create(const ProxySettings& InSettings, TMap<IdType, ProxyType>& OutProxyMap, FAudioModulationSystem& OutModSystem, TFunction<void(ProxyType&)> OnCreateProxy)
 		{
-			const IdType ObjectId = static_cast<IdType>(InObject.GetId());
-			TProxyHandle<IdType, ProxyType, ProxiedObject> NewHandle = Get(InObject.GetId(), InProxyMap);
+			const IdType ObjectId = static_cast<IdType>(InSettings.GetId());
+			TProxyHandle<IdType, ProxyType, ProxySettings> NewHandle = Get(InSettings.GetId(), OutProxyMap);
 			if (!NewHandle.IsValid())
 			{
-				UE_LOG(LogAudioModulation, Verbose, TEXT("Proxy created: Id '%u' for object '%s'."), NewHandle.Id, *InObject.GetName());
-				InProxyMap.Add(ObjectId, ProxyType(InObject, OutModSystem));
-				NewHandle = TProxyHandle<IdType, ProxyType, ProxiedObject>(ObjectId, InProxyMap);
+				UE_LOG(LogAudioModulation, Verbose, TEXT("Proxy created: Id '%u' for object '%s'."), NewHandle.Id, *InSettings.GetName());
+				OutProxyMap.Add(ObjectId, ProxyType(InSettings, OutModSystem));
+				NewHandle = TProxyHandle<IdType, ProxyType, ProxySettings>(ObjectId, OutProxyMap);
 				OnCreateProxy(NewHandle.FindProxy());
 			}
 
 			return NewHandle;
 		}
 
-		FORCEINLINE bool operator ==(const TProxyHandle<IdType, ProxyType, ProxiedObject>& InHandle) const
+		FORCEINLINE bool operator ==(const TProxyHandle<IdType, ProxyType, ProxySettings>& InHandle) const
 		{
 			return InHandle.Id == Id;
 		}
 
-		FORCEINLINE TProxyHandle<IdType, ProxyType, ProxiedObject>& operator =(const TProxyHandle<IdType, ProxyType, ProxiedObject>& InHandle)
+		FORCEINLINE TProxyHandle<IdType, ProxyType, ProxySettings>& operator =(const TProxyHandle<IdType, ProxyType, ProxySettings>& InHandle)
 		{
 			// 1. If local proxy valid prior to move, cache to DecRef
 			ProxyType* ProxyToDecRef = nullptr;
@@ -190,7 +190,7 @@ namespace AudioModulation
 			return *this;
 		}
 
-		FORCEINLINE TProxyHandle<IdType, ProxyType, ProxiedObject>& operator =(TProxyHandle<IdType, ProxyType, ProxiedObject>&& InHandle)
+		FORCEINLINE TProxyHandle<IdType, ProxyType, ProxySettings>& operator =(TProxyHandle<IdType, ProxyType, ProxySettings>&& InHandle)
 		{
 			// 1. If local proxy valid prior to move, cache to DecRef
 			ProxyType* ProxyToDecRef = nullptr;
@@ -273,7 +273,7 @@ namespace AudioModulation
 		}
 	};
 
-	template<typename IdType, typename ProxyType, typename ProxiedObject>
+	template<typename IdType, typename ProxyType, typename ProxySettings>
 	class TModulatorProxyRefType : public TModulatorBase<IdType>
 	{
 	protected:
@@ -354,6 +354,6 @@ namespace AudioModulation
 		}
 
 	private:
-		friend class TProxyHandle<IdType, ProxyType, ProxiedObject>;
+		friend class TProxyHandle<IdType, ProxyType, ProxySettings>;
 	};
 } // namespace AudioModulation
\ No newline at end of file
diff --git a/Engine/Plugins/Runtime/AudioModulation/Source/AudioModulation/Public/AudioModulation.h b/Engine/Plugins/Runtime/AudioModulation/Source/AudioModulation/Public/AudioModulation.h
--- a/Engine/Plugins/Runtime/AudioModulation/Source/AudioModulation/Public/AudioModulation.h
+++ b/Engine/Plugins/Runtime/AudioModulation/Source/AudioModulation/Public/AudioModulation.h
@@ -1,12 +1,15 @@
 // Copyright Epic Games, Inc. All Rights Reserved.
 #pragma once
 
+#include "AudioDeviceManager.h"
 #include "CoreMinimal.h"
 #include "HAL/ThreadSafeBool.h"
 #include "IAudioExtensionPlugin.h"
 #include "IAudioModulation.h"
 #include "Modules/ModuleInterface.h"
-#include "SoundModulationPatch.h"
+#include "SoundControlBus.h"
+#include "SoundControlBusMix.h"
+#include "SoundModulationGenerator.h"
 #include "Stats/Stats.h"
 
 
@@ -19,20 +22,39 @@ DECLARE_CYCLE_STAT_EXTERN(TEXT("Process Controls"), STAT_AudioModulationProcessC
 
 namespace AudioModulation
 {
+	// Forward Declarations
 	class FAudioModulationSystem;
 
 	class AUDIOMODULATION_API FAudioModulation : public IAudioModulation
 	{
 	public:
 		FAudioModulation();
-		virtual ~FAudioModulation() = default;
+		virtual ~FAudioModulation();
 
 		//~ Begin IAudioModulation implementation
-		virtual Audio::FModulationParameter GetParameter(FName InParamName);
+		virtual Audio::FModulationParameter GetParameter(FName InParamName) override;
 		virtual void Initialize(const FAudioPluginInitializationParams& InitializationParams) override;
 
 		virtual void OnAuditionEnd() override;
 
+		void ActivateBus(const USoundControlBus& InBus);
+		void ActivateBusMix(const USoundControlBusMix& InBusMix);
+		void ActivateGenerator(const USoundModulationGenerator& InGenerator);
+
+		void DeactivateBus(const USoundControlBus& InBus);
+		void DeactivateBusMix(const USoundControlBusMix& InBusMix);
+		void DeactivateAllBusMixes();
+		void DeactivateGenerator(const USoundModulationGenerator& InGenerator);
+
+		void SaveMixToProfile(const USoundControlBusMix& InBusMix, const int32 InProfileIndex);
+		TArray<FSoundControlBusMixStage> LoadMixFromProfile(const int32 InProfileIndex, USoundControlBusMix& OutBusMix);
+
+		void UpdateMix(const TArray<FSoundControlBusMixStage>& InStages, USoundControlBusMix& InOutMix, bool bInUpdateObject = false, float InFadeTime = -1.0f);
+		void UpdateMix(const USoundControlBusMix& InMix, float InFadeTime = -1.0f);
+		void UpdateMixByFilter(const FString& InAddressFilter, const TSubclassOf<USoundModulationParameter>& InParamClassFilter, USoundModulationParameter* InParamFilter, float Value, float FadeTime, USoundControlBusMix& InOutMix, bool bInUpdateObject = false);
+
+		void SoloBusMix(const USoundControlBusMix& InBusMix);
+
 #if !UE_BUILD_SHIPPING
 		virtual bool OnPostHelp(FCommonViewportClient* ViewportClient, const TCHAR* Stream) override;
 		virtual int32 OnRenderStat(FViewport* Viewport, FCanvas* Canvas, int32 X, int32 Y, const UFont& Font, const FVector* ViewLocation, const FRotator* ViewRotation) override;
@@ -44,8 +66,6 @@ namespace AudioModulation
 		virtual void UpdateModulator(const USoundModulatorBase& InModulator) override;
 		//~ End IAudioModulation implementation
 
-		FAudioModulationSystem* GetModulationSystem();
-
 	protected:
 		virtual Audio::FModulatorTypeId RegisterModulator(Audio::FModulatorHandleId InHandleId, const USoundModulatorBase* InModulatorBase, Audio::FModulationParameter& OutParameter) override;
 		virtual void RegisterModulator(Audio::FModulatorHandleId InHandleId, Audio::FModulatorId InModulatorId) override;
@@ -53,8 +73,24 @@ namespace AudioModulation
 		virtual void UnregisterModulator(const Audio::FModulatorHandle& InHandle) override;
 
 	private:
-		TUniquePtr<FAudioModulationSystem> ModSystem;
+		FAudioModulationSystem* ModSystem = nullptr;
 	};
+
+	static void IterateModulationImpl(TUniqueFunction<void(FAudioModulation&)> InFunction)
+	{
+		if (FAudioDeviceManager* DeviceManager = FAudioDeviceManager::Get())
+		{
+			TArray<FAudioDevice*> Devices = DeviceManager->GetAudioDevices();
+			DeviceManager->IterateOverAllDevices([ModFunction = MoveTemp(InFunction)](Audio::FDeviceId DeviceId, FAudioDevice* AudioDevice)
+			{
+				if (AudioDevice && AudioDevice->IsModulationPluginEnabled() && AudioDevice->ModulationInterface.IsValid())
+				{
+					auto ModulationInterface = static_cast<AudioModulation::FAudioModulation*>(AudioDevice->ModulationInterface.Get());
+					ModFunction(*ModulationInterface);
+				}
+			});
+		}
+	}
 } // namespace AudioModulation
 
 class FAudioModulationPluginFactory : public IAudioModulationFactory
diff --git a/Engine/Plugins/Runtime/AudioModulation/Source/AudioModulation/Public/AudioModulationStatics.h b/Engine/Plugins/Runtime/AudioModulation/Source/AudioModulation/Public/AudioModulationStatics.h
--- a/Engine/Plugins/Runtime/AudioModulation/Source/AudioModulation/Public/AudioModulationStatics.h
+++ b/Engine/Plugins/Runtime/AudioModulation/Source/AudioModulation/Public/AudioModulationStatics.h
@@ -1,6 +1,7 @@
 // Copyright Epic Games, Inc. All Rights Reserved.
 #pragma once
 
+#include "AudioModulation.h"
 #include "CoreMinimal.h"
 #include "Kismet/BlueprintFunctionLibrary.h"
 #include "SoundControlBus.h"
@@ -32,7 +33,7 @@ public:
 	/**
 	 * Returns modulation implementation associated with the provided world
 	 */
-	static AudioModulation::FAudioModulationSystem* GetModulationSystem(UWorld* World);
+	static AudioModulation::FAudioModulation* GetModulation(UWorld* World);
 
 	/** Activates a bus. Does nothing if an instance of the provided bus is already active
 	 * @param Bus - Bus to activate
@@ -52,14 +53,14 @@ public:
 	)
 	static void ActivateBusMix(const UObject* WorldContextObject, USoundControlBusMix* BusMix);
 
-	/** Activates a bus modulator (eg. LFO). Does nothing if an instance of the provided modulator is already active
+	/** Activates a modulation generator. Does nothing if an instance of the provided generator is already active
 	 * @param Modulator - Modulator to activate
 	 */
-	UFUNCTION(BlueprintCallable, Category = "Audio", DisplayName = "Activate Control Bus Modulator", meta = (
+	UFUNCTION(BlueprintCallable, Category = "Audio", DisplayName = "Activate Modulation Generator", meta = (
 		WorldContext = "WorldContextObject", 
-		Keywords = "activate modulation modulator lfo")
+		Keywords = "activate modulation modulator generator lfo")
 	)
-	static void ActivateBusModulator(const UObject* WorldContextObject, USoundModulationGenerator* Modulator);
+	static void ActivateGenerator(const UObject* WorldContextObject, USoundModulationGenerator* Generator);
 
 	/** Creates a modulation bus with the provided default value.
 	 * @param Name - Name of bus
@@ -73,27 +74,6 @@ public:
 	)
 	static USoundControlBus* CreateBus(const UObject* WorldContextObject, FName Name, USoundModulationParameter* Parameter, bool Activate = true);
 
-	/** Creates an LFO modulator.
-	 * @param Name - Name of LFO
-	 * @param Amplitude - Amplitude of new LFO.
-	 * @param Frequency - Frequency of new LFO.
-	 * @param Offset - Offset of new LFO.
-	 * @param Activate - Whether or not to activate lfo on creation.
-	 */
-	UFUNCTION(BlueprintCallable, Category = "Audio", DisplayName = "Create Control Bus LFO", meta = (
-		AdvancedDisplay = "5",
-		WorldContext = "WorldContextObject",
-		Keywords = "make create lfo modulation modulator")
-	)
-	static USoundModulationGeneratorLFO* CreateLFO(
-		const UObject* WorldContextObject, 
-		FName Name, 
-		float Amplitude, 
-		float Frequency, 
-		float Offset,
-		bool Activate = true
-	);
-
 	/** Creates a stage used to mix a control bus.
 	 * @param Bus - Bus stage is in charge of applying mix value to.
 	 * @param Stages - Value for added bus stage to target when mix is active.
@@ -144,17 +124,17 @@ public:
 	)
 	static void DeactivateBusMix(const UObject* WorldContextObject, USoundControlBusMix* BusMix);
 
-	/** Deactivates a bus modulator. Does nothing if an instance of the provided bus mix is already inactive
-	 * @param Modulator - Modulator to activate
+	/** Deactivates a modulation generator. Does nothing if an instance of the provided generator is already inactive
+	 * @param Generator - Generator to activate
 	 * @param Scope - Scope of modulator
 	 */
 	UFUNCTION(BlueprintCallable, Category = "Audio", DisplayName = "Deactivate Control Bus Modulator", meta = (
 		WorldContext = "WorldContextObject",
 		Keywords = "deactivate bus modulation modulator")
 )
-	static void DeactivateBusModulator(const UObject* WorldContextObject, USoundModulationGenerator* Modulator);
+	static void DeactivateGenerator(const UObject* WorldContextObject, USoundModulationGenerator* Generator);
 
-	/** Saves control bus mix to a profile, serialized to an ini file.  If mix is loaded, uses current proxy's state. 
+	/** Saves control bus mix to a profile, serialized to an ini file.  If mix is loaded, uses current proxy's state.
 	 * If not, uses default UObject representation.
 	 * @param BusMix - Mix object to serialize to profile .ini.
 	 * @param ProfileIndex - Index of profile, allowing multiple profiles can be saved for single mix object. If 0, saves to default ini profile (no suffix).
diff --git a/Engine/Plugins/Runtime/AudioModulation/Source/AudioModulation/Public/Generators/SoundModulationLFO.h b/Engine/Plugins/Runtime/AudioModulation/Source/AudioModulation/Public/Generators/SoundModulationLFO.h
new file mode 100644
--- /dev/null
+++ b/Engine/Plugins/Runtime/AudioModulation/Source/AudioModulation/Public/Generators/SoundModulationLFO.h
@@ -0,0 +1,148 @@
+// Copyright Epic Games, Inc. All Rights Reserved.
+#pragma once
+
+#include "CoreMinimal.h"
+#include "DSP/LFO.h"
+#include "IAudioModulation.h"
+#include "SoundModulationGenerator.h"
+#include "Templates/SharedPointer.h"
+#include "UObject/ObjectMacros.h"
+#include "UObject/Object.h"
+
+#include "SoundModulationLFO.generated.h"
+
+
+UENUM(BlueprintType)
+enum class FSoundModulationLFOShape : uint8
+{
+	Sine			  UMETA(DisplayName = "Sine"),
+	UpSaw			  UMETA(DisplayName = "Saw (Up)"),
+	DownSaw			  UMETA(DisplayName = "Saw (Down)"),
+	Square			  UMETA(DisplayName = "Square"),
+	Triangle		  UMETA(DisplayName = "Triangle"),
+	Exponential		  UMETA(DisplayName = "Exponential"),
+	RandomSampleHold  UMETA(DisplayName = "Random"),
+
+	COUNT UMETA(Hidden)
+};
+
+USTRUCT(BlueprintType)
+struct FSoundModulationLFOParams
+{
+	GENERATED_USTRUCT_BODY()
+
+	/** Shape of oscillating waveform */
+	UPROPERTY(EditAnywhere, Category = Modulation, BlueprintReadWrite)
+	FSoundModulationLFOShape Shape;
+
+	/** Amplitude of oscillator */
+	UPROPERTY(EditAnywhere, Category = Modulation, BlueprintReadWrite, meta = (UIMin = "0", UIMax = "1", ClampMin = "0", ClampMax = "1"))
+	float Amplitude = 0.5f;
+
+	/** Frequency of oscillator */
+	UPROPERTY(EditAnywhere, Category = Modulation, BlueprintReadWrite, meta = (UIMin = "0", UIMax = "20", ClampMin = "0", ClampMax = "20"))
+	float Frequency = 1.0f;
+
+	/** Offset of oscillator */
+	UPROPERTY(EditAnywhere, Category = Modulation, BlueprintReadWrite, meta = (UIMin = "0", UIMax = "1", ClampMin = "0", ClampMax = "1"))
+	float Offset = 0.5;
+
+	/** Whether or not to loop the oscillation more than once */
+	UPROPERTY(EditAnywhere, Category = Modulation, BlueprintReadWrite)
+	bool bLooping = 1;
+
+	/** If true, bypasses LFO bus from being modulated by parameters, patches, or mixed (LFO remains active and computed). */
+	UPROPERTY(EditAnywhere, Category = Modulation, BlueprintReadWrite)
+	bool bBypass = 0;
+};
+
+
+namespace AudioModulation
+{
+	class FLFOGenerator : public IGenerator
+	{
+		public:
+			FLFOGenerator(const FSoundModulationLFOParams& InParams)
+				: Params(InParams)
+			{
+				LFO.SetGain(Params.Amplitude);
+				LFO.SetFrequency(Params.Frequency);
+				LFO.SetMode(Params.bLooping ? Audio::ELFOMode::Type::Sync : Audio::ELFOMode::OneShot);
+
+				static_assert(static_cast<int32>(FSoundModulationLFOShape::COUNT) == static_cast<int32>(Audio::ELFO::Type::NumLFOTypes), "LFOShape/ELFO Type mismatch");
+				LFO.SetType(static_cast<Audio::ELFO::Type>(Params.Shape));
+				LFO.Start();
+			}
+
+			virtual ~FLFOGenerator() = default;
+
+			virtual float GetValue() const override
+			{
+				return Value;
+			}
+
+			virtual bool IsBypassed() const override
+			{
+				return Params.bBypass;
+			}
+
+			virtual void Update(double InElapsed) override
+			{
+				if (InElapsed > 0.0f && LFO.GetFrequency() > 0.0f)
+				{
+					const float SampleRate = static_cast<float>(1.0 / InElapsed);
+					LFO.SetSampleRate(SampleRate);
+					LFO.Update();
+					Value = LFO.Generate() + Params.Offset;
+				}
+			}
+
+#if !UE_BUILD_SHIPPING
+		static const FString DebugName;
+
+		virtual void GetDebugValues(TArray<FString>& OutDebugValues) const override;
+		virtual const FString& GetDebugName() const override;
+#endif // !UE_BUILD_SHIPPING
+
+	protected:
+		Audio::FLFO LFO;
+		float Value = 1.0f;
+		FSoundModulationLFOParams Params;
+	};
+} // namespace AudioModulation
+
+
+UCLASS(BlueprintType, hidecategories = Object, editinlinenew, MinimalAPI)
+class USoundModulationGeneratorLFO : public USoundModulationGenerator
+{
+	GENERATED_BODY()
+
+public:
+	UPROPERTY(EditAnywhere, Category = Modulation, BlueprintReadWrite, meta = (ShowOnlyInnerProperties))
+	FSoundModulationLFOParams Params;
+
+#if !UE_BUILD_SHIPPING
+	static const TArray<FString>& GetDebugCategories()
+	{
+		static const TArray<FString> Categories =
+		{
+			TEXT("Value"),
+			TEXT("Gain"),
+			TEXT("Frequency"),
+			TEXT("Offset"),
+			TEXT("Curve")
+		};
+		return Categories;
+	}
+
+	static const FString& GetDebugName();
+#endif // !UE_BUILD_SHIPPING
+
+
+	virtual AudioModulation::FGeneratorPtr CreateInstance(Audio::FDeviceId InDeviceId) const override
+	{
+		using namespace AudioModulation;
+		auto NewGenerator = MakeShared<FLFOGenerator, ESPMode::ThreadSafe>(Params);
+		return StaticCastSharedRef<IGenerator>(NewGenerator);
+	}
+};
diff --git a/Engine/Plugins/Runtime/AudioModulation/Source/AudioModulation/Public/SoundControlBus.h b/Engine/Plugins/Runtime/AudioModulation/Source/AudioModulation/Public/SoundControlBus.h
--- a/Engine/Plugins/Runtime/AudioModulation/Source/AudioModulation/Public/SoundControlBus.h
+++ b/Engine/Plugins/Runtime/AudioModulation/Source/AudioModulation/Public/SoundControlBus.h
@@ -6,8 +6,8 @@
 #include "AudioDevice.h"
 #include "CoreMinimal.h"
 #include "IAudioModulation.h"
+#include "SoundModulationGenerator.h"
 #include "SoundModulationParameter.h"
-#include "SoundModulationGeneratorLFO.h"
 #include "UObject/Object.h"
 #include "UObject/ObjectMacros.h"
 
@@ -40,8 +40,8 @@ public:
 	UPROPERTY(EditAnywhere, Category = Mix, BlueprintReadWrite, meta = (EditCondition = "bOverrideAddress"))
 	FString Address;
 
-	UPROPERTY(EditAnywhere, Category = Generators, BlueprintReadWrite, meta = (DisplayName = "Generators"))
-	TArray<USoundModulationGenerator*> Modulators;
+	UPROPERTY(EditAnywhere, Category = Generators, BlueprintReadWrite)
+	TArray<USoundModulationGenerator*> Generators;
 
 	UPROPERTY(EditAnywhere, Category = General, BlueprintReadOnly)
 	USoundModulationParameter* Parameter;
diff --git a/Engine/Plugins/Runtime/AudioModulation/Source/AudioModulation/Public/SoundModulationGenerator.h b/Engine/Plugins/Runtime/AudioModulation/Source/AudioModulation/Public/SoundModulationGenerator.h
--- a/Engine/Plugins/Runtime/AudioModulation/Source/AudioModulation/Public/SoundModulationGenerator.h
+++ b/Engine/Plugins/Runtime/AudioModulation/Source/AudioModulation/Public/SoundModulationGenerator.h
@@ -1,26 +1,72 @@
 // Copyright Epic Games, Inc. All Rights Reserved.
 #pragma once
 
+#include "AudioDeviceManager.h"
 #include "CoreMinimal.h"
 #include "DSP/LFO.h"
 #include "IAudioModulation.h"
 #include "SoundModulationValue.h"
+#include "Templates/SharedPointer.h"
 #include "UObject/ObjectMacros.h"
 #include "UObject/Object.h"
 
 #include "SoundModulationGenerator.generated.h"
 
 
+namespace AudioModulation
+{
+	class AUDIOMODULATION_API IGenerator
+	{
+	public:
+		IGenerator(Audio::FDeviceId InDeviceId = INDEX_NONE)
+			: AudioDeviceId(InDeviceId)
+		{
+		}
+		virtual ~IGenerator() = default;
+
+		virtual float GetValue() const { return 1.0f; }
+		virtual bool IsBypassed() const { return true; };
+		virtual void Update(double InElapsed) { };
+
+
+#if !UE_BUILD_SHIPPING
+		virtual void GetDebugValues(TArray<FString>& OutDebugValues) const { }
+
+		// Required for instance look-up in factory registration
+		virtual const FString& GetDebugName() const { static const FString Invalid; return Invalid; }
+#endif // !UE_BUILD_SHIPPING
+
+		void AudioRenderThreadCommand(TUniqueFunction<void()>&& InCommand);
+		void PumpCommands();
+
+	protected:
+		Audio::FDeviceId AudioDeviceId = INDEX_NONE;
+
+	private:
+		TQueue<TUniqueFunction<void()>> CommandQueue;
+	};
+
+	using FGeneratorPtr = TSharedPtr<IGenerator, ESPMode::ThreadSafe>;
+} // namespace AudioModulation
+
 /**
- * Base class for modulators that algoithmically generate values that can effect
+ * Base class for modulators that algorithmically generate values that can effect
  * various endpoints (ex. Control Buses & Parameter Destinations)
  */
-UCLASS(hideCategories = Object, abstract, MinimalAPI)
-class USoundModulationGenerator : public USoundModulatorBase
+UCLASS(hideCategories = Object, abstract)
+class AUDIOMODULATION_API USoundModulationGenerator : public USoundModulatorBase
 {
-	GENERATED_UCLASS_BODY()
+	GENERATED_BODY()
+
+public:
+	virtual AudioModulation::FGeneratorPtr CreateInstance(Audio::FDeviceId AudioDeviceId) const
+	{
+		return nullptr;
+	}
 
 #if WITH_EDITOR
-		virtual void PostEditChangeProperty(FPropertyChangedEvent& InPropertyChangedEvent) override;
+	virtual void PostEditChangeProperty(FPropertyChangedEvent& InPropertyChangedEvent) override;
 #endif // WITH_EDITOR
+
+	virtual void BeginDestroy() override;
 };
diff --git a/Engine/Plugins/Runtime/AudioModulation/Source/AudioModulation/Public/SoundModulationGeneratorLFO.h b/Engine/Plugins/Runtime/AudioModulation/Source/AudioModulation/Public/SoundModulationGeneratorLFO.h
deleted file mode 100644
diff --git a/Engine/Plugins/Runtime/AudioModulation/Source/AudioModulationEditor/Private/AssetTypeActions/AssetTypeActions_SoundModulationGenerator.cpp b/Engine/Plugins/Runtime/AudioModulation/Source/AudioModulationEditor/Private/AssetTypeActions/AssetTypeActions_SoundModulationGenerator.cpp
new file mode 100644
--- /dev/null
+++ b/Engine/Plugins/Runtime/AudioModulation/Source/AudioModulationEditor/Private/AssetTypeActions/AssetTypeActions_SoundModulationGenerator.cpp
@@ -0,0 +1,22 @@
+// Copyright Epic Games, Inc. All Rights Reserved.
+#include "AssetTypeActions_SoundModulationGenerator.h"
+
+#include "SoundModulationGenerator.h"
+
+
+#define LOCTEXT_NAMESPACE "AssetTypeActions"
+UClass* FAssetTypeActions_SoundModulationGenerator::GetSupportedClass() const
+{
+	return USoundModulationGenerator::StaticClass();
+}
+
+const TArray<FText>& FAssetTypeActions_SoundModulationGenerator::GetSubMenus() const
+{
+	static const TArray<FText> SubMenus
+	{
+		LOCTEXT("AssetSoundModulationSubMenu", "Modulation")
+	};
+
+	return SubMenus;
+}
+#undef LOCTEXT_NAMESPACE // AssetTypeActions
\ No newline at end of file
diff --git a/Engine/Plugins/Runtime/AudioModulation/Source/AudioModulationEditor/Private/AssetTypeActions/AssetTypeActions_SoundModulationGenerator.h b/Engine/Plugins/Runtime/AudioModulation/Source/AudioModulationEditor/Private/AssetTypeActions/AssetTypeActions_SoundModulationGenerator.h
new file mode 100644
--- /dev/null
+++ b/Engine/Plugins/Runtime/AudioModulation/Source/AudioModulationEditor/Private/AssetTypeActions/AssetTypeActions_SoundModulationGenerator.h
@@ -0,0 +1,17 @@
+// Copyright Epic Games, Inc. All Rights Reserved.
+#pragma once
+
+#include "AssetTypeActions_Base.h"
+#include "AudioModulationStyle.h"
+#include "CoreMinimal.h"
+
+class FAssetTypeActions_SoundModulationGenerator : public FAssetTypeActions_Base
+{
+public:
+	// IAssetTypeActions Implementation
+	virtual FText GetName() const override { return NSLOCTEXT("AssetTypeActions", "AssetTypeActions_SoundModulationGenerator", "Modulation Generator"); }
+	virtual FColor GetTypeColor() const override { return UAudioModulationStyle::GetModulationGeneratorColor(); }
+	virtual const TArray<FText>& GetSubMenus() const override;
+	virtual UClass* GetSupportedClass() const override;
+	virtual uint32 GetCategories() override { return EAssetTypeCategories::Sounds; }
+};
diff --git a/Engine/Plugins/Runtime/AudioModulation/Source/AudioModulationEditor/Private/AssetTypeActions/AssetTypeActions_SoundModulationGeneratorLFO.cpp b/Engine/Plugins/Runtime/AudioModulation/Source/AudioModulationEditor/Private/AssetTypeActions/AssetTypeActions_SoundModulationGeneratorLFO.cpp
deleted file mode 100644
diff --git a/Engine/Plugins/Runtime/AudioModulation/Source/AudioModulationEditor/Private/AssetTypeActions/AssetTypeActions_SoundModulationGeneratorLFO.h b/Engine/Plugins/Runtime/AudioModulation/Source/AudioModulationEditor/Private/AssetTypeActions/AssetTypeActions_SoundModulationGeneratorLFO.h
deleted file mode 100644
diff --git a/Engine/Plugins/Runtime/AudioModulation/Source/AudioModulationEditor/Private/AudioModulationEditor.cpp b/Engine/Plugins/Runtime/AudioModulation/Source/AudioModulationEditor/Private/AudioModulationEditor.cpp
--- a/Engine/Plugins/Runtime/AudioModulation/Source/AudioModulationEditor/Private/AudioModulationEditor.cpp
+++ b/Engine/Plugins/Runtime/AudioModulation/Source/AudioModulationEditor/Private/AudioModulationEditor.cpp
@@ -4,7 +4,7 @@
 #include "AssetRegistryModule.h"
 #include "AssetTypeActions/AssetTypeActions_SoundControlBus.h"
 #include "AssetTypeActions/AssetTypeActions_SoundControlBusMix.h"
-#include "AssetTypeActions/AssetTypeActions_SoundModulationGeneratorLFO.h"
+#include "AssetTypeActions/AssetTypeActions_SoundModulationGenerator.h"
 #include "AssetTypeActions/AssetTypeActions_SoundModulationParameter.h"
 #include "AssetTypeActions/AssetTypeActions_SoundModulationPatch.h"
 #include "Editors/ModulationPatchCurveEditorViewStacked.h"
@@ -43,6 +43,7 @@ namespace AudioModulationEditor
 	}
 } // namespace AudioModulationEditor
 
+
 FAudioModulationEditorModule::FAudioModulationEditorModule()
 {
 }
@@ -84,13 +85,13 @@ void FAudioModulationEditorModule::StartupModule()
 
 	AudioModulationEditor::AddAssetAction<FAssetTypeActions_SoundControlBus>(AssetTools, AssetActions);
 	AudioModulationEditor::AddAssetAction<FAssetTypeActions_SoundControlBusMix>(AssetTools, AssetActions);
-	AudioModulationEditor::AddAssetAction<FAssetTypeActions_SoundModulationGeneratorLFO>(AssetTools, AssetActions);
+	AudioModulationEditor::AddAssetAction<FAssetTypeActions_SoundModulationGenerator>(AssetTools, AssetActions);
 	AudioModulationEditor::AddAssetAction<FAssetTypeActions_SoundModulationParameter>(AssetTools, AssetActions);
 	AudioModulationEditor::AddAssetAction<FAssetTypeActions_SoundModulationPatch>(AssetTools, AssetActions);
 
 	SetIcon(TEXT("SoundControlBus"));
 	SetIcon(TEXT("SoundControlBusMix"));
-	SetIcon(TEXT("SoundModulationGeneratorLFO"));
+	SetIcon(TEXT("SoundModulationGenerator"));
 	SetIcon(TEXT("SoundModulationPatch"));
 	SetIcon(TEXT("SoundModulationParameter"));
 
diff --git a/Engine/Plugins/Runtime/AudioModulation/Source/AudioModulationEditor/Private/Factories/SoundControlBusMixFactory.h b/Engine/Plugins/Runtime/AudioModulation/Source/AudioModulationEditor/Private/Factories/SoundControlBusMixFactory.h
--- a/Engine/Plugins/Runtime/AudioModulation/Source/AudioModulationEditor/Private/Factories/SoundControlBusMixFactory.h
+++ b/Engine/Plugins/Runtime/AudioModulation/Source/AudioModulationEditor/Private/Factories/SoundControlBusMixFactory.h
@@ -1,16 +1,13 @@
 // Copyright Epic Games, Inc. All Rights Reserved.
-
-//=============================================================================
-// USoundModulationBusFactory
-//=============================================================================
-
 #pragma once
 
 #include "CoreMinimal.h"
-#include "UObject/ObjectMacros.h"
 #include "Factories/Factory.h"
+#include "UObject/ObjectMacros.h"
+
 #include "SoundControlBusMixFactory.generated.h"
 
+
 UCLASS(hidecategories=Object, MinimalAPI)
 class USoundControlBusMixFactory : public UFactory
 {
diff --git a/Engine/Plugins/Runtime/AudioModulation/Source/AudioModulationEditor/Private/Factories/SoundModulationGeneratorFactory.cpp b/Engine/Plugins/Runtime/AudioModulation/Source/AudioModulationEditor/Private/Factories/SoundModulationGeneratorFactory.cpp
new file mode 100644
--- /dev/null
+++ b/Engine/Plugins/Runtime/AudioModulation/Source/AudioModulationEditor/Private/Factories/SoundModulationGeneratorFactory.cpp
@@ -0,0 +1,84 @@
+// Copyright Epic Games, Inc. All Rights Reserved.
+
+#include "SoundModulationGeneratorFactory.h"
+
+#include "ClassViewerFilter.h"
+#include "ClassViewerModule.h"
+#include "Kismet2/SClassPickerDialog.h"
+#include "SoundModulationGenerator.h"
+#include "Templates/SharedPointer.h"
+
+#define LOCTEXT_NAMESPACE "AudioModulation"
+
+
+namespace AudioModulation
+{
+	namespace Editor
+	{
+		class FGeneratorClassViewerFilter : public IClassViewerFilter
+		{
+		public:
+			TSet<const UClass*> AllowedChildrenOfClasses;
+			EClassFlags DisallowedClassFlags;
+
+			FGeneratorClassViewerFilter()
+				: DisallowedClassFlags(CLASS_Abstract | CLASS_Deprecated | CLASS_NewerVersionExists)
+			{
+				AllowedChildrenOfClasses.Add(USoundModulationGenerator::StaticClass());
+			}
+
+			virtual bool IsClassAllowed(const FClassViewerInitializationOptions& InInitOptions, const UClass* InClass, TSharedRef<FClassViewerFilterFuncs> InFilterFuncs) override
+			{
+				return !InClass->HasAnyClassFlags(DisallowedClassFlags)
+					&& InFilterFuncs->IfInChildOfClassesSet(AllowedChildrenOfClasses, InClass) != EFilterReturn::Failed;
+			}
+
+			virtual bool IsUnloadedClassAllowed(const FClassViewerInitializationOptions& InInitOptions, const TSharedRef<const IUnloadedBlueprintData> InUnloadedClassData, TSharedRef<FClassViewerFilterFuncs> InFilterFuncs) override
+			{
+				return !InUnloadedClassData->HasAnyClassFlags(DisallowedClassFlags)
+					&& InFilterFuncs->IfInChildOfClassesSet(AllowedChildrenOfClasses, InUnloadedClassData) != EFilterReturn::Failed;
+			}
+		};
+	} // namespace Editor
+} // namespace AudioModulation
+
+USoundModulationGeneratorFactory::USoundModulationGeneratorFactory(const FObjectInitializer& ObjectInitializer)
+	: Super(ObjectInitializer)
+{
+	SupportedClass = USoundModulationGenerator::StaticClass();
+	bCreateNew     = true;
+	bEditorImport  = false;
+	bEditAfterNew  = true;
+	GeneratorClass = nullptr;
+}
+
+bool USoundModulationGeneratorFactory::ConfigureProperties()
+{
+	using namespace AudioModulation::Editor;
+
+	FClassViewerModule& ClassViewerModule = FModuleManager::LoadModuleChecked<FClassViewerModule>("ClassViewer");
+
+	FClassViewerInitializationOptions Options;
+	Options.ClassFilter = MakeShared<FGeneratorClassViewerFilter>();
+	Options.Mode = EClassViewerMode::ClassPicker;
+
+	const FText TitleText = LOCTEXT("CreateModulationGeneratorClassSelect", "Select Generator Class");
+
+	UClass* PickedClass = nullptr;
+	const bool bPressedOk = SClassPickerDialog::PickClass(TitleText, Options, PickedClass, USoundModulationGenerator::StaticClass());
+	GeneratorClass = bPressedOk ? PickedClass : nullptr;
+
+	return bPressedOk;
+}
+
+UObject* USoundModulationGeneratorFactory::FactoryCreateNew(UClass* InClass, UObject* InParent, FName InName, EObjectFlags Flags, UObject* Context, FFeedbackContext* Warn)
+{
+	if (GeneratorClass)
+	{
+		return NewObject<USoundModulationGenerator>(InParent, GeneratorClass, InName, Flags);
+	}
+
+	return nullptr;
+}
+
+#undef LOCTEXT_NAMESPACE // AudioModulation
diff --git a/Engine/Plugins/Runtime/AudioModulation/Source/AudioModulationEditor/Private/Factories/SoundModulationGeneratorFactory.h b/Engine/Plugins/Runtime/AudioModulation/Source/AudioModulationEditor/Private/Factories/SoundModulationGeneratorFactory.h
new file mode 100644
--- /dev/null
+++ b/Engine/Plugins/Runtime/AudioModulation/Source/AudioModulationEditor/Private/Factories/SoundModulationGeneratorFactory.h
@@ -0,0 +1,23 @@
+// Copyright Epic Games, Inc. All Rights Reserved.
+#pragma once
+
+#include "CoreMinimal.h"
+#include "Factories/Factory.h"
+#include "UObject/ObjectMacros.h"
+#include "SoundModulationGenerator.h"
+
+#include "SoundModulationGeneratorFactory.generated.h"
+
+UCLASS(hidecategories=Object, MinimalAPI)
+class USoundModulationGeneratorFactory : public UFactory
+{
+	GENERATED_UCLASS_BODY()
+
+	UPROPERTY()
+	TSubclassOf<USoundModulationGenerator> GeneratorClass;
+
+	//~ Begin UFactory Interface
+	virtual bool ConfigureProperties() override;
+	virtual UObject* FactoryCreateNew(UClass* Class, UObject* InParent, FName Name, EObjectFlags Flags, UObject* Context, FFeedbackContext* Warn) override;
+	//~ Begin UFactory Interface	
+};
diff --git a/Engine/Plugins/Runtime/AudioModulation/Source/AudioModulationEditor/Private/Factories/SoundModulationGeneratorLFOFactory.cpp b/Engine/Plugins/Runtime/AudioModulation/Source/AudioModulationEditor/Private/Factories/SoundModulationGeneratorLFOFactory.cpp
deleted file mode 100644
diff --git a/Engine/Plugins/Runtime/AudioModulation/Source/AudioModulationEditor/Private/Factories/SoundModulationGeneratorLFOFactory.h b/Engine/Plugins/Runtime/AudioModulation/Source/AudioModulationEditor/Private/Factories/SoundModulationGeneratorLFOFactory.h
deleted file mode 100644
