Modulation Generator Class Template

diff --git a/Engine/Plugins/Runtime/AudioModulation/Source/AudioModulationEditor/Private/AudioModulationClassTemplates.cpp b/Engine/Plugins/Runtime/AudioModulation/Source/AudioModulationEditor/Private/AudioModulationClassTemplates.cpp
new file mode 100644
--- /dev/null
+++ b/Engine/Plugins/Runtime/AudioModulation/Source/AudioModulationEditor/Private/AudioModulationClassTemplates.cpp
@@ -0,0 +1,14 @@
+// Copyright Epic Games, Inc. All Rights Reserved.
+#include "AudioModulationClassTemplates.h"
+
+USoundModulationClassTemplate::USoundModulationClassTemplate(const FObjectInitializer& ObjectInitializer)
+	: UPluginClassTemplate(ObjectInitializer)
+{
+	PluginName = TEXT("AudioModulation");
+}
+
+USoundModulationGeneratorClassTemplate::USoundModulationGeneratorClassTemplate(const FObjectInitializer& ObjectInitializer)
+	: USoundModulationClassTemplate(ObjectInitializer)
+{
+	SetGeneratedBaseClass(USoundModulationGenerator::StaticClass());
+}
diff --git a/Engine/Plugins/Runtime/AudioModulation/Source/AudioModulationEditor/Public/AudioModulationClassTemplates.h b/Engine/Plugins/Runtime/AudioModulation/Source/AudioModulationEditor/Public/AudioModulationClassTemplates.h
new file mode 100644
--- /dev/null
+++ b/Engine/Plugins/Runtime/AudioModulation/Source/AudioModulationEditor/Public/AudioModulationClassTemplates.h
@@ -0,0 +1,22 @@
+// Copyright Epic Games, Inc. All Rights Reserved.
+#pragma once
+
+#include "CoreMinimal.h"
+
+#include "ClassTemplateEditorSubsystem.h"
+#include "SoundModulationGenerator.h"
+
+#include "AudioModulationClassTemplates.generated.h"
+
+
+UCLASS(Abstract)
+class AUDIOMODULATIONEDITOR_API USoundModulationClassTemplate : public UPluginClassTemplate
+{
+	GENERATED_UCLASS_BODY()
+};
+
+UCLASS()
+class AUDIOMODULATIONEDITOR_API USoundModulationGeneratorClassTemplate : public USoundModulationClassTemplate
+{
+	GENERATED_UCLASS_BODY()
+};
