Modulation Envelope Follower Generator: AudioBus implementation MVP

diff --git a/Engine/Plugins/Runtime/AudioModulation/Source/AudioModulation/Private/Generators/SoundModulationEnvelopeFollower.cpp b/Engine/Plugins/Runtime/AudioModulation/Source/AudioModulation/Private/Generators/SoundModulationEnvelopeFollower.cpp
new file mode 100644
--- /dev/null
+++ b/Engine/Plugins/Runtime/AudioModulation/Source/AudioModulation/Private/Generators/SoundModulationEnvelopeFollower.cpp
@@ -0,0 +1,105 @@
+// Copyright Epic Games, Inc. All Rights Reserved.
+
+#include "Generators/SoundModulationEnvelopeFollower.h"
+
+#include "AudioDeviceManager.h"
+#include "AudioMixerDevice.h"
+#include "AudioModulation.h"
+
+
+namespace AudioModulation
+{
+	const FString FEnvelopeFollowerGenerator::DebugName = TEXT("EnvelopeFollower");
+
+	FEnvelopeFollowerGenerator::FEnvelopeFollowerGenerator(const FEnvelopeFollowerGeneratorParams& InParams, Audio::FDeviceId InDeviceId)
+		: IGenerator(InDeviceId)
+		, Params(InParams)
+	{
+		using namespace Audio;
+
+		if (Params.AudioBus)
+		{
+			AudioRenderThreadCommand([this, DeviceId = InDeviceId, BusId = Params.AudioBus->GetUniqueID()]()
+			{
+				if (FAudioDeviceManager* DeviceManager = FAudioDeviceManager::Get())
+				{
+					FAudioDevice* AudioDevice = DeviceManager->GetAudioDeviceRaw(DeviceId);
+					if (FMixerDevice* MixerDevice = static_cast<FMixerDevice*>(AudioDevice))
+					{
+						AudioBusPatch = MixerDevice->AddPatchForAudioBus(BusId, Params.Gain);
+						EnvelopeFollower = FEnvelopeFollower(MixerDevice->SampleRate, Params.AttackTime, Params.ReleaseTime, Audio::EPeakMode::Peak);
+					}
+				}
+			});
+		}
+	}
+
+	float FEnvelopeFollowerGenerator::GetValue() const
+	{
+		return EnvelopeFollower.GetCurrentValue();
+	}
+
+	bool FEnvelopeFollowerGenerator::IsBypassed() const
+	{
+		return Params.bBypass;
+	}
+
+	void FEnvelopeFollowerGenerator::Update(double InElapsed)
+	{
+		float Amp = 0.0f;
+
+		if (AudioBusPatch.IsValid())
+		{
+			int32 NumSamples = AudioBusPatch->GetNumSamplesAvailable();
+			if (NumSamples > 0)
+			{
+				TempBuffer.Reset();
+				TempBuffer.AddZeroed(NumSamples);
+
+				AudioBusPatch->PopAudio(TempBuffer.GetData(), NumSamples, true /* bUseLatestAudio */);
+
+				// If NumSamples is below 4 samples, just take first sample
+				static const int32 SimdMultiple = 4;
+				if (NumSamples % SimdMultiple == 0)
+				{
+					Amp = Audio::GetAverageAmplitude(TempBuffer.GetData(), NumSamples);
+				}
+				// If NumSamples is not a multiple of 4, do it the non-SIMD way
+				else
+				{
+					Amp = 0.0f;
+					for (int32 i = 0; i < NumSamples; ++i)
+					{
+						Amp += TempBuffer[i];
+					}
+					Amp /= NumSamples;
+				}
+			}
+		}
+
+		EnvelopeFollower.ProcessAudio(Amp);
+	}
+
+#if !UE_BUILD_SHIPPING
+	void FEnvelopeFollowerGenerator::GetDebugValues(TArray<FString>& OutDebugValues) const
+	{
+		OutDebugValues.Add(FString::Printf(TEXT("%.4f"), EnvelopeFollower.GetCurrentValue()));
+		OutDebugValues.Add(FString::Printf(TEXT("%.4f"), Params.Gain));
+		OutDebugValues.Add(FString::Printf(TEXT("%.4f"), Params.AttackTime));
+		OutDebugValues.Add(FString::Printf(TEXT("%.4f"), Params.ReleaseTime));
+	}
+
+	const FString& FEnvelopeFollowerGenerator::GetDebugName() const
+	{
+		return DebugName;
+	}
+#endif // !UE_BUILD_SHIPPING
+} // namespace AudioModulation
+
+#if !UE_BUILD_SHIPPING
+const FString& USoundModulationGeneratorEnvelopeFollower::GetDebugName()
+{
+	using namespace AudioModulation;
+	return FEnvelopeFollowerGenerator::DebugName;
+}
+#endif // !UE_BUILD_SHIPPING
diff --git a/Engine/Plugins/Runtime/AudioModulation/Source/AudioModulation/Public/Generators/SoundModulationEnvelopeFollower.h b/Engine/Plugins/Runtime/AudioModulation/Source/AudioModulation/Public/Generators/SoundModulationEnvelopeFollower.h
new file mode 100644
--- /dev/null
+++ b/Engine/Plugins/Runtime/AudioModulation/Source/AudioModulation/Public/Generators/SoundModulationEnvelopeFollower.h
@@ -0,0 +1,110 @@
+// Copyright Epic Games, Inc. All Rights Reserved.
+#pragma once
+
+#include "CoreMinimal.h"
+#include "SoundModulationGenerator.h"
+
+#include "AudioDeviceManager.h"
+#include "DSP/BufferVectorOperations.h"
+#include "DSP/EnvelopeFollower.h"
+#include "DSP/MultithreadedPatching.h"
+#include "IAudioModulation.h"
+#include "SoundModulationGenerator.h"
+#include "Templates/SharedPointer.h"
+#include "UObject/Object.h"
+#include "UObject/ObjectMacros.h"
+
+#include "SoundModulationEnvelopeFollower.generated.h"
+
+
+// Forward Declarations
+class UAudioBus;
+
+USTRUCT(BlueprintType)
+struct FEnvelopeFollowerGeneratorParams
+{
+	GENERATED_USTRUCT_BODY()
+
+	/** AudioBus to follow amplitude of and generate modulation control signal from. */
+	UPROPERTY(EditAnywhere, Category = Modulation, BlueprintReadWrite)
+	UAudioBus* AudioBus = nullptr;
+
+	/** Gain to apply to amplitude signal. */
+	UPROPERTY(EditAnywhere, Category = Modulation, BlueprintReadWrite)
+	float Gain = 1.0f;
+
+	/** Attack time of envelope response (in sec) */
+	UPROPERTY(EditAnywhere, Category = Modulation, BlueprintReadWrite, meta = (ClampMin = 0.0f))
+	float AttackTime = 0.010f;
+
+	/** Release time of envelope response (in sec) */
+	UPROPERTY(EditAnywhere, Category = Modulation, BlueprintReadWrite, meta = (ClampMin = 0.0f))
+	float ReleaseTime = 0.100f;
+
+	/** If true, bypasses generator from being modulated by parameters, patches, or mixed (remains active and computed). */
+	UPROPERTY(EditAnywhere, Category = Modulation, BlueprintReadWrite)
+	bool bBypass = false;
+};
+
+namespace AudioModulation
+{
+	class AUDIOMODULATION_API FEnvelopeFollowerGenerator : public IGenerator
+	{
+		public:
+			FEnvelopeFollowerGenerator(const FEnvelopeFollowerGeneratorParams& InParams, Audio::FDeviceId InDeviceId);
+			virtual ~FEnvelopeFollowerGenerator() = default;
+
+			virtual float GetValue() const override;
+			virtual bool IsBypassed() const override;
+			virtual void Update(double InElapsed) override;
+
+#if !UE_BUILD_SHIPPING
+			static const FString DebugName;
+
+			virtual void GetDebugValues(TArray<FString>& OutDebugValues) const override;
+			virtual const FString& GetDebugName() const override;
+#endif // !UE_BUILD_SHIPPING
+
+	protected:
+		FEnvelopeFollowerGeneratorParams Params;
+
+	private:
+		Audio::FPatchOutputStrongPtr AudioBusPatch;
+		Audio::FAlignedFloatBuffer TempBuffer;
+		Audio::FEnvelopeFollower EnvelopeFollower;
+	};
+} // namespace AudioModulation
+
+UCLASS(hidecategories = Object, BlueprintType, editinlinenew, meta = (DisplayName = "Envelope Follower Generator"))
+class AUDIOMODULATION_API USoundModulationGeneratorEnvelopeFollower : public USoundModulationGenerator
+{
+	GENERATED_BODY()
+
+public:
+	UPROPERTY(EditAnywhere, Category = Modulation, BlueprintReadWrite, meta = (ShowOnlyInnerProperties))
+	FEnvelopeFollowerGeneratorParams Params;
+
+#if !UE_BUILD_SHIPPING
+	static const TArray<FString>& GetDebugCategories()
+	{
+		static const TArray<FString> Categories =
+		{
+			TEXT("Value"),
+			TEXT("Gain"),
+			TEXT("Attack"),
+			TEXT("Release"),
+		};
+		return Categories;
+	}
+
+	static const FString& GetDebugName();
+#endif // !UE_BUILD_SHIPPING
+
+	virtual AudioModulation::FGeneratorPtr CreateInstance(Audio::FDeviceId InDeviceId) const override
+	{
+		using namespace AudioModulation;
+
+		auto NewGenerator = MakeShared<FEnvelopeFollowerGenerator, ESPMode::ThreadSafe>(Params, InDeviceId);
+		return StaticCastSharedRef<AudioModulation::IGenerator>(NewGenerator);
+	}
+};
