Rename 'linear' space to 'normalized' now that parameters are actually normalizing and not just clamping (per docs discussion)

diff --git a/Engine/Plugins/Runtime/AudioModulation/Source/AudioModulation/Private/AudioModulationSystem.cpp b/Engine/Plugins/Runtime/AudioModulation/Source/AudioModulation/Private/AudioModulationSystem.cpp
--- a/Engine/Plugins/Runtime/AudioModulation/Source/AudioModulation/Private/AudioModulationSystem.cpp
+++ b/Engine/Plugins/Runtime/AudioModulation/Source/AudioModulation/Private/AudioModulationSystem.cpp
@@ -274,8 +274,9 @@ namespace AudioModulation
 						Parameter.ParameterName = InParamName;
 						Parameter.bRequiresConversion = Param->RequiresUnitConversion();
 						Parameter.MixFunction = Param->GetMixFunction();
+						Parameter.UnitDisplayName = Param->Settings.UnitDisplayName;
 						Parameter.UnitFunction = Param->GetUnitConversionFunction();
-						Parameter.LinearFunction = Param->GetLinearConversionFunction();
+						Parameter.NormalizedFunction = Param->GetNormalizedConversionFunction();
 						Parameter.DefaultValue = Param->GetUnitDefault();
 						Parameter.MinValue = Param->GetUnitMin();
 						Parameter.MaxValue = Param->GetUnitMax();
diff --git a/Engine/Plugins/Runtime/AudioModulation/Source/AudioModulation/Private/SoundControlBusProxy.h b/Engine/Plugins/Runtime/AudioModulation/Source/AudioModulation/Private/SoundControlBusProxy.h
--- a/Engine/Plugins/Runtime/AudioModulation/Source/AudioModulation/Private/SoundControlBusProxy.h
+++ b/Engine/Plugins/Runtime/AudioModulation/Source/AudioModulation/Private/SoundControlBusProxy.h
@@ -28,7 +28,7 @@ namespace AudioModulation
 		FControlBusSettings(const USoundControlBus& InBus, Audio::FDeviceId InDeviceId)
 			: TModulatorBase<FBusId>(InBus.GetName(), InBus.GetUniqueID())
 			, bBypass()
-			, DefaultValue(InBus.GetDefaultLinearValue())
+			, DefaultValue(InBus.GetDefaultNormalizedValue())
 			, MixFunction(InBus.GetMixFunction())
 		{
 			for (const USoundModulationGenerator* Generator : InBus.Generators)
diff --git a/Engine/Plugins/Runtime/AudioModulation/Source/AudioModulation/Private/SoundModulationParameter.cpp b/Engine/Plugins/Runtime/AudioModulation/Source/AudioModulation/Private/SoundModulationParameter.cpp
--- a/Engine/Plugins/Runtime/AudioModulation/Source/AudioModulation/Private/SoundModulationParameter.cpp
+++ b/Engine/Plugins/Runtime/AudioModulation/Source/AudioModulation/Private/SoundModulationParameter.cpp
@@ -3,19 +3,19 @@
 
 
 #if WITH_EDITOR
-void USoundModulationParameter::RefreshLinearValue()
+void USoundModulationParameter::RefreshNormalizedValue()
 {
-	const float NewLinearValue = ConvertUnitToLinear(Settings.ValueUnit);
-	const float NewLinearValueClamped = FMath::Clamp(NewLinearValue, 0.0f, 1.0f);
-	if (!FMath::IsNearlyEqual(NewLinearValueClamped, Settings.ValueLinear))
+	const float NewNormalizedValue = ConvertUnitToNormalized(Settings.ValueUnit);
+	const float NewNormalizedValueClamped = FMath::Clamp(NewNormalizedValue, 0.0f, 1.0f);
+	if (!FMath::IsNearlyEqual(NewNormalizedValueClamped, Settings.ValueNormalized))
 	{
-		Settings.ValueLinear = NewLinearValueClamped;
+		Settings.ValueNormalized = NewNormalizedValueClamped;
 	}
 }
 
 void USoundModulationParameter::RefreshUnitValue()
 {
-	const float NewUnitValue = ConvertLinearToUnit(Settings.ValueLinear);
+	const float NewUnitValue = ConvertNormalizedToUnit(Settings.ValueNormalized);
 	const float NewUnitValueClamped = FMath::Clamp(NewUnitValue, GetUnitMin(), GetUnitMax());
 	if (!FMath::IsNearlyEqual(NewUnitValueClamped, Settings.ValueUnit))
 	{
@@ -42,7 +42,7 @@ Audio::FModulationUnitConvertFunction USoundModulationParameterFrequencyBase::Ge
 	};
 }
 
-Audio::FModulationLinearConversionFunction USoundModulationParameterFrequencyBase::GetLinearConversionFunction() const
+Audio::FModulationNormalizedConversionFunction USoundModulationParameterFrequencyBase::GetNormalizedConversionFunction() const
 {
 	return [InUnitMin = GetUnitMin(), InUnitMax = GetUnitMax()](float* RESTRICT OutValueBuffer, int32 InNumSamples)
 	{
@@ -67,7 +67,7 @@ Audio::FModulationMixFunction USoundModulationParameterHPFFrequency::GetMixFunct
 USoundModulationParameterHPFFrequency::USoundModulationParameterHPFFrequency(const FObjectInitializer& ObjectInitializer)
 	: Super(ObjectInitializer)
 {
-	Settings.ValueLinear = 0.0f;
+	Settings.ValueNormalized = 0.0f;
 
 #if WITH_EDITORONLY_DATA
 	Settings.ValueUnit = GetUnitDefault();
@@ -101,7 +101,7 @@ Audio::FModulationUnitConvertFunction USoundModulationParameterScaled::GetUnitCo
 	};
 }
 
-Audio::FModulationLinearConversionFunction USoundModulationParameterScaled::GetLinearConversionFunction() const
+Audio::FModulationNormalizedConversionFunction USoundModulationParameterScaled::GetNormalizedConversionFunction() const
 {
 	return [InUnitMin = UnitMin, InUnitMax = UnitMax](float* RESTRICT OutValueBuffer, int32 InNumSamples)
 	{
@@ -126,7 +126,7 @@ float USoundModulationParameterScaled::GetUnitMax() const
 USoundModulationParameterBipolar::USoundModulationParameterBipolar(const FObjectInitializer& ObjectInitializer)
 	: Super(ObjectInitializer)
 {
-	Settings.ValueLinear = 0.5f;
+	Settings.ValueNormalized = 0.5f;
 
 #if WITH_EDITORONLY_DATA
 	Settings.ValueUnit = GetUnitDefault();
@@ -160,7 +160,7 @@ Audio::FModulationUnitConvertFunction USoundModulationParameterBipolar::GetUnitC
 	};
 }
 
-Audio::FModulationLinearConversionFunction USoundModulationParameterBipolar::GetLinearConversionFunction() const
+Audio::FModulationNormalizedConversionFunction USoundModulationParameterBipolar::GetNormalizedConversionFunction() const
 {
 	return [InUnitRange = UnitRange](float* RESTRICT OutValueBuffer, int32 InNumSamples)
 	{
@@ -199,7 +199,7 @@ Audio::FModulationUnitConvertFunction USoundModulationParameterVolume::GetUnitCo
 	};
 }
 
-Audio::FModulationLinearConversionFunction USoundModulationParameterVolume::GetLinearConversionFunction() const
+Audio::FModulationNormalizedConversionFunction USoundModulationParameterVolume::GetNormalizedConversionFunction() const
 {
 	return [InUnitMin = GetUnitMin()](float* RESTRICT OutValueBuffer, int32 InNumSamples)
 	{
diff --git a/Engine/Plugins/Runtime/AudioModulation/Source/AudioModulation/Private/SoundModulationPatchProxy.h b/Engine/Plugins/Runtime/AudioModulation/Source/AudioModulation/Private/SoundModulationPatchProxy.h
--- a/Engine/Plugins/Runtime/AudioModulation/Source/AudioModulation/Private/SoundModulationPatchProxy.h
+++ b/Engine/Plugins/Runtime/AudioModulation/Source/AudioModulation/Private/SoundModulationPatchProxy.h
@@ -80,7 +80,7 @@ namespace AudioModulation
 		{
 			if (InPatch.OutputParameter)
 			{
-				DefaultOutputValue = InPatch.OutputParameter->Settings.ValueLinear;
+				DefaultOutputValue = InPatch.OutputParameter->Settings.ValueNormalized;
 				MixFunction = InPatch.OutputParameter->GetMixFunction();
 			}
 
diff --git a/Engine/Plugins/Runtime/AudioModulation/Source/AudioModulation/Public/SoundControlBus.h b/Engine/Plugins/Runtime/AudioModulation/Source/AudioModulation/Public/SoundControlBus.h
--- a/Engine/Plugins/Runtime/AudioModulation/Source/AudioModulation/Public/SoundControlBus.h
+++ b/Engine/Plugins/Runtime/AudioModulation/Source/AudioModulation/Public/SoundControlBus.h
@@ -56,7 +56,7 @@ public:
 	virtual void BeginDestroy() override;
 	virtual const Audio::FModulationMixFunction& GetMixFunction() const;
 
-	virtual float GetDefaultLinearValue() const { return Parameter ? Parameter->Settings.ValueLinear : 1.0f; }
+	virtual float GetDefaultNormalizedValue() const { return Parameter ? Parameter->Settings.ValueNormalized : 1.0f; }
 
 	virtual FName GetOutputParameterName() const override
 	{
diff --git a/Engine/Plugins/Runtime/AudioModulation/Source/AudioModulation/Public/SoundModulationParameter.h b/Engine/Plugins/Runtime/AudioModulation/Source/AudioModulation/Public/SoundModulationParameter.h
--- a/Engine/Plugins/Runtime/AudioModulation/Source/AudioModulation/Public/SoundModulationParameter.h
+++ b/Engine/Plugins/Runtime/AudioModulation/Source/AudioModulation/Public/SoundModulationParameter.h
@@ -15,14 +15,14 @@ struct AUDIOMODULATION_API FSoundModulationParameterSettings
 	GENERATED_USTRUCT_BODY()
 
 	/** 
-	  * Linear default value of modulator. To ensure bypass functionality of mixing, patching, and modulating 
+	  * Normalized, unitless default value of modulator. To ensure bypass functionality of mixing, patching, and modulating 
 	  * functions as anticipated, value should be selected such that GetMixFunction (see USoundModulationParameter)
 	  * reduces to an identity function (i.e. function acts as a "pass-through" for all values in the range [0.0, 1.0]).
-	  * If GetMixFunction performs the mathmatical operation f(x1, x2), then the default ValueLinear should result in
-	  * f(x1, d) = x1 where d is ValueLinear.
+	  * If GetMixFunction performs the mathematical operation f(x1, x2), then the default ValueNormalized should result in
+	  * f(x1, d) = x1 where d is ValueNormalized.
 	  */
 	UPROPERTY(EditAnywhere, Category = General, meta = (ClampMin = "0.0", ClampMax = "1.0"))
-	float ValueLinear = 1.0f;
+	float ValueNormalized = 1.0f;
 
 #if WITH_EDITORONLY_DATA
 	/** (Optional) Text name of parameter's unit */
@@ -65,8 +65,8 @@ public:
 		return MixFunction;
 	}
 
-	/** Function used to convert linear value to unit value */
-	virtual Audio::FModulationLinearConversionFunction  GetUnitConversionFunction() const
+	/** Function used to convert normalized, unitless value to unit value */
+	virtual Audio::FModulationNormalizedConversionFunction  GetUnitConversionFunction() const
 	{
 		static const Audio::FModulationUnitConvertFunction ConversionFunction = [](float* RESTRICT OutValueBuffer, int32 InNumSamples)
 		{
@@ -75,36 +75,36 @@ public:
 		return ConversionFunction;
 	}
 
-	/** Function used to convert unit value to linear value */
-	virtual Audio::FModulationLinearConversionFunction GetLinearConversionFunction() const
+	/** Function used to convert unit value to normalized, unitless value */
+	virtual Audio::FModulationNormalizedConversionFunction GetNormalizedConversionFunction() const
 	{
-		static const Audio::FModulationLinearConversionFunction ConversionFunction = [](float* RESTRICT OutValueBuffer, int32 InNumSamples)
+		static const Audio::FModulationNormalizedConversionFunction ConversionFunction = [](float* RESTRICT OutValueBuffer, int32 InNumSamples)
 		{
 		};
 
 		return ConversionFunction;
 	}
 
-	/** Converts linear [0.0f, 1.0f] value to unit value. */
-	virtual float ConvertLinearToUnit(float InLinearValue) const final
+	/** Converts normalized, unitless value [0.0f, 1.0f] to unit value. */
+	virtual float ConvertNormalizedToUnit(float InNormalizedValue) const final
 	{
-		float UnitValue = InLinearValue;
+		float UnitValue = InNormalizedValue;
 		GetUnitConversionFunction()(&UnitValue, 1);
 		return UnitValue;
 	}
 
-	/** Converts unit value to linear [0.0f, 1.0f] value. */
-	virtual float ConvertUnitToLinear(float InUnitValue) const final
+	/** Converts unit value to unitless, normalized value [0.0f, 1.0f]. */
+	virtual float ConvertUnitToNormalized(float InUnitValue) const final
 	{
-		float LinearValue = InUnitValue;
-		GetLinearConversionFunction()(&LinearValue, 1);
-		return LinearValue;
+		float NormalizedValue = InUnitValue;
+		GetNormalizedConversionFunction()(&NormalizedValue, 1);
+		return NormalizedValue;
 	}
 
 	/** Returns default unit value (works with and without editor loaded) */
 	virtual float GetUnitDefault() const
 	{
-		return ConvertLinearToUnit(Settings.ValueLinear);
+		return ConvertNormalizedToUnit(Settings.ValueNormalized);
 	}
 
 	virtual float GetUnitMin() const
@@ -118,12 +118,12 @@ public:
 	}
 
 #if WITH_EDITOR
-	void RefreshLinearValue();
+	void RefreshNormalizedValue();
 	void RefreshUnitValue();
 #endif // WITH_EDITOR
 };
 
-// Modulation Parameter that scales linear value to explicit unit minimum and maximum.
+// Linearly scaled value between unit minimum and maximum.
 UCLASS(BlueprintType, MinimalAPI)
 class USoundModulationParameterScaled : public USoundModulationParameter
 {
@@ -140,12 +140,12 @@ public:
 
 	virtual bool RequiresUnitConversion() const override;
 	virtual Audio::FModulationUnitConvertFunction GetUnitConversionFunction() const override;
-	virtual Audio::FModulationLinearConversionFunction GetLinearConversionFunction() const override;
+	virtual Audio::FModulationNormalizedConversionFunction GetNormalizedConversionFunction() const override;
 	virtual float GetUnitMin() const override;
 	virtual float GetUnitMax() const override;
 };
 
-// Modulation Parameter that scales linear value to logarithmic frequency unit space.
+// Modulation Parameter that scales normalized, unitless value to logarithmic frequency unit space.
 UCLASS(BlueprintType, MinimalAPI, abstract)
 class USoundModulationParameterFrequencyBase : public USoundModulationParameter
 {
@@ -154,10 +154,10 @@ class USoundModulationParameterFrequencyBase : public USoundModulationParameter
 public:
 	virtual bool RequiresUnitConversion() const override;
 	virtual Audio::FModulationUnitConvertFunction GetUnitConversionFunction() const override;
-	virtual Audio::FModulationLinearConversionFunction GetLinearConversionFunction() const override;
+	virtual Audio::FModulationNormalizedConversionFunction GetNormalizedConversionFunction() const override;
 };
 
-// Modulation Parameter that scales linear value to logarithmic frequency unit space with provided minimum and maximum.
+// Modulation Parameter that scales normalized, unitless value to logarithmic frequency unit space with provided minimum and maximum.
 UCLASS(BlueprintType, MinimalAPI)
 class USoundModulationParameterFrequency : public USoundModulationParameterFrequencyBase
 {
@@ -183,7 +183,7 @@ public:
 	}
 };
 
-// Modulation Parameter that scales linear value to logarithmic frequency unit space with standard filter min and max frequency set.
+// Modulation Parameter that scales normalized, unitless value to logarithmic frequency unit space with standard filter min and max frequency set.
 UCLASS(BlueprintType, MinimalAPI)
 class USoundModulationParameterFilterFrequency : public USoundModulationParameterFrequencyBase
 {
@@ -201,7 +201,7 @@ public:
 	}
 };
 
-// Modulation Parameter that scales linear value to logarithmic frequency unit space with standard filter min and max frequency set.
+// Modulation Parameter that scales normalized, unitless value to logarithmic frequency unit space with standard filter min and max frequency set.
 // Mixes by taking the minimum (i.e. aggressive) filter frequency of all active modulators.
 UCLASS(BlueprintType, MinimalAPI)
 class USoundModulationParameterLPFFrequency : public USoundModulationParameterFilterFrequency
@@ -212,7 +212,7 @@ public:
 	virtual Audio::FModulationMixFunction GetMixFunction() const override;
 };
 
-// Modulation Parameter that scales linear value to logarithmic frequency unit space with standard filter min and max frequency set.
+// Modulation Parameter that scales normalized, unitless value to logarithmic frequency unit space with standard filter min and max frequency set.
 // Mixes by taking the maximum (i.e. aggressive) filter frequency of all active modulators.
 UCLASS(BlueprintType, MinimalAPI)
 class USoundModulationParameterHPFFrequency : public USoundModulationParameterFilterFrequency
@@ -223,7 +223,7 @@ public:
 	virtual Audio::FModulationMixFunction GetMixFunction() const override;
 };
 
-// Modulation Parameter that scales linear value to bipolar range. Mixes multiplicatively.
+// Modulation Parameter that scales normalized, unitless value to bipolar range. Mixes multiplicatively.
 UCLASS(BlueprintType, MinimalAPI)
 class USoundModulationParameterBipolar : public USoundModulationParameter
 {
@@ -237,7 +237,7 @@ public:
 	virtual bool RequiresUnitConversion() const override;
 	virtual Audio::FModulationMixFunction GetMixFunction() const override;
 	virtual Audio::FModulationUnitConvertFunction GetUnitConversionFunction() const override;
-	virtual Audio::FModulationLinearConversionFunction GetLinearConversionFunction() const override;
+	virtual Audio::FModulationNormalizedConversionFunction GetNormalizedConversionFunction() const override;
 	virtual float GetUnitMax() const override;
 	virtual float GetUnitMin() const override;
 };
@@ -254,7 +254,7 @@ public:
 
 	virtual bool RequiresUnitConversion() const override;
 	virtual Audio::FModulationUnitConvertFunction GetUnitConversionFunction() const override;
-	virtual Audio::FModulationLinearConversionFunction GetLinearConversionFunction() const override;
+	virtual Audio::FModulationNormalizedConversionFunction GetNormalizedConversionFunction() const override;
 	virtual float GetUnitMin() const override;
 	virtual float GetUnitMax() const override;
 };
diff --git a/Engine/Plugins/Runtime/AudioModulation/Source/AudioModulationEditor/Private/Editors/ModulationPatchCurveEditorViewStacked.cpp b/Engine/Plugins/Runtime/AudioModulation/Source/AudioModulationEditor/Private/Editors/ModulationPatchCurveEditorViewStacked.cpp
--- a/Engine/Plugins/Runtime/AudioModulation/Source/AudioModulationEditor/Private/Editors/ModulationPatchCurveEditorViewStacked.cpp
+++ b/Engine/Plugins/Runtime/AudioModulation/Source/AudioModulationEditor/Private/Editors/ModulationPatchCurveEditorViewStacked.cpp
@@ -24,8 +24,8 @@ namespace PatchCurveViewUtils
 {
 	void FormatLabel(const USoundModulationParameter& InParameter, const FNumberFormattingOptions& InNumFormatOptions, FText& InOutLabel)
 	{
-		const float LinearValue = FCString::Atof(*InOutLabel.ToString());
-		const float UnitValue = InParameter.ConvertLinearToUnit(LinearValue);
+		const float NormalizedValue = FCString::Atof(*InOutLabel.ToString());
+		const float UnitValue = InParameter.ConvertNormalizedToUnit(NormalizedValue);
 		FText UnitLabel = FText::AsNumber(UnitValue, &InNumFormatOptions);
 		InOutLabel = FText::Format(LOCTEXT("ModulationPatchCurveView_UnitFormat", "{0} ({1})"), UnitLabel, InOutLabel);
 	}
@@ -44,7 +44,7 @@ FModPatchCurveEditorModel::FModPatchCurveEditorModel(FRichCurve& InRichCurve, UO
 
 void FModPatchCurveEditorModel::Refresh(EModPatchOutputEditorCurveSource InSource, int32 InInputIndex)
 {
-	InputAxisName = LOCTEXT("ModulationCurveDisplayTitle_Linear", "Linear");
+	InputAxisName = LOCTEXT("ModulationCurveDisplayTitle_Normalized", "Normalized");
 	FText OutputAxisName = InputAxisName;
 
 	FSoundControlModulationInput* Input = nullptr;
diff --git a/Engine/Plugins/Runtime/AudioModulation/Source/AudioModulationEditor/Private/Layouts/SoundControlBusMixStageLayout.cpp b/Engine/Plugins/Runtime/AudioModulation/Source/AudioModulationEditor/Private/Layouts/SoundControlBusMixStageLayout.cpp
--- a/Engine/Plugins/Runtime/AudioModulation/Source/AudioModulationEditor/Private/Layouts/SoundControlBusMixStageLayout.cpp
+++ b/Engine/Plugins/Runtime/AudioModulation/Source/AudioModulationEditor/Private/Layouts/SoundControlBusMixStageLayout.cpp
@@ -50,23 +50,23 @@ namespace AudioModulationEditorUtils
 		return nullptr;
 	}
 
-	void HandleConvertLinearToUnit(TSharedRef<IPropertyHandle> BusHandle, TSharedRef<IPropertyHandle> LinearValueHandle, TSharedRef<IPropertyHandle> UnitValueHandle)
+	void HandleConvertNormalizedToUnit(TSharedRef<IPropertyHandle> BusHandle, TSharedRef<IPropertyHandle> NormalizedValueHandle, TSharedRef<IPropertyHandle> UnitValueHandle)
 	{
 		if (USoundModulationParameter* Parameter = AudioModulationEditorUtils::GetParameterFromBus(BusHandle))
 		{
 			if (Parameter->RequiresUnitConversion())
 			{
-				float LinearValue = 1.0f;
-				if (LinearValueHandle->GetValue(LinearValue) == FPropertyAccess::Success)
+				float NormalizedValue = 1.0f;
+				if (NormalizedValueHandle->GetValue(NormalizedValue) == FPropertyAccess::Success)
 				{
-					const float UnitValue = Parameter->ConvertLinearToUnit(LinearValue);
+					const float UnitValue = Parameter->ConvertNormalizedToUnit(NormalizedValue);
 					UnitValueHandle->SetValue(UnitValue, EPropertyValueSetFlags::NotTransactable);
 				}
 			}
 		}
 	}
 
-	void HandleConvertUnitToLinear(TSharedRef<IPropertyHandle> BusHandle, TSharedRef<IPropertyHandle> UnitValueHandle, TSharedRef<IPropertyHandle> LinearValueHandle)
+	void HandleConvertUnitToNormalized(TSharedRef<IPropertyHandle> BusHandle, TSharedRef<IPropertyHandle> UnitValueHandle, TSharedRef<IPropertyHandle> NormalizedValueHandle)
 	{
 		if (USoundModulationParameter* Parameter = AudioModulationEditorUtils::GetParameterFromBus(BusHandle))
 		{
@@ -75,8 +75,8 @@ namespace AudioModulationEditorUtils
 				float UnitValue = 1.0f;
 				if (UnitValueHandle->GetValue(UnitValue) == FPropertyAccess::Success)
 				{
-					const float LinearValue = Parameter->ConvertUnitToLinear(UnitValue);
-					LinearValueHandle->SetValue(LinearValue);
+					const float NormalizedValue = Parameter->ConvertUnitToNormalized(UnitValue);
+					NormalizedValueHandle->SetValue(NormalizedValue);
 				}
 			}
 		}
@@ -112,21 +112,21 @@ void FSoundControlBusMixStageLayoutCustomization::CustomizeChildren(TSharedRef<I
 
 	ChildBuilder.AddProperty(BusHandle);
 
-	TSharedRef<IPropertyHandle> LinearValueHandle = StructPropertyHandle->GetChildHandle(GET_MEMBER_NAME_CHECKED(FSoundModulationMixValue, TargetValue)).ToSharedRef();
+	TSharedRef<IPropertyHandle> NormalizedValueHandle = StructPropertyHandle->GetChildHandle(GET_MEMBER_NAME_CHECKED(FSoundModulationMixValue, TargetValue)).ToSharedRef();
 	TSharedRef<IPropertyHandle> UnitValueHandle = StructPropertyHandle->GetChildHandle(GET_MEMBER_NAME_CHECKED(FSoundModulationMixValue, TargetUnitValue)).ToSharedRef();
 
 	// When editor opens, set unit value in case bus unit has changed while editor was closed.
-	AudioModulationEditorUtils::HandleConvertLinearToUnit(BusHandle, LinearValueHandle, UnitValueHandle);
+	AudioModulationEditorUtils::HandleConvertNormalizedToUnit(BusHandle, NormalizedValueHandle, UnitValueHandle);
 
-	LinearValueHandle->SetOnPropertyValueChanged(FSimpleDelegate::CreateLambda([BusHandle, LinearValueHandle, UnitValueHandle]()
+	NormalizedValueHandle->SetOnPropertyValueChanged(FSimpleDelegate::CreateLambda([BusHandle, NormalizedValueHandle, UnitValueHandle]()
 	{
-		AudioModulationEditorUtils::HandleConvertLinearToUnit(BusHandle, LinearValueHandle, UnitValueHandle);
+		AudioModulationEditorUtils::HandleConvertNormalizedToUnit(BusHandle, NormalizedValueHandle, UnitValueHandle);
 	}));
 
-	UnitValueHandle->SetOnPropertyValueChanged(FSimpleDelegate::CreateLambda([BusHandle, UnitValueHandle, LinearValueHandle]()
+	UnitValueHandle->SetOnPropertyValueChanged(FSimpleDelegate::CreateLambda([BusHandle, UnitValueHandle, NormalizedValueHandle]()
 	{
-		AudioModulationEditorUtils::HandleConvertUnitToLinear(BusHandle, UnitValueHandle, LinearValueHandle);
-		AudioModulationEditorUtils::HandleConvertLinearToUnit(BusHandle, LinearValueHandle, UnitValueHandle);
+		AudioModulationEditorUtils::HandleConvertUnitToNormalized(BusHandle, UnitValueHandle, NormalizedValueHandle);
+		AudioModulationEditorUtils::HandleConvertNormalizedToUnit(BusHandle, NormalizedValueHandle, UnitValueHandle);
 	}));
 
 	ChildBuilder.AddCustomRow(StructPropertyHandle->GetPropertyDisplayName())
@@ -171,7 +171,7 @@ void FSoundControlBusMixStageLayoutCustomization::CustomizeChildren(TSharedRef<I
 				.Padding(4.0f, 0.0f, 0.0f, 0.0f)
 				.VAlign(VAlign_Center)
 				[
-					LinearValueHandle->CreatePropertyValueWidget()
+					NormalizedValueHandle->CreatePropertyValueWidget()
 				]
 			+ SHorizontalBox::Slot()
 				.Padding(4.0f, 0.0f, 0.0f, 0.0f)
@@ -180,7 +180,7 @@ void FSoundControlBusMixStageLayoutCustomization::CustomizeChildren(TSharedRef<I
 				[
 					SNew(STextBlock)
 					.Font(IDetailLayoutBuilder::GetDetailFont())
-					.Text(LOCTEXT("SoundModulationControl_UnitValueLinear", "Linear"))
+					.Text(LOCTEXT("SoundModulationControl_UnitValueNormalized", "Normalized"))
 				]
 	]
 	.Visibility(TAttribute<EVisibility>::Create([BusHandle]()
@@ -193,7 +193,7 @@ void FSoundControlBusMixStageLayoutCustomization::CustomizeChildren(TSharedRef<I
 		return EVisibility::Hidden;
 	}));
 
-	ChildBuilder.AddProperty(LinearValueHandle)
+	ChildBuilder.AddProperty(NormalizedValueHandle)
 	.Visibility(TAttribute<EVisibility>::Create([BusHandle]()
 	{
 		if (USoundModulationParameter* Parameter = AudioModulationEditorUtils::GetParameterFromBus(BusHandle))
diff --git a/Engine/Plugins/Runtime/AudioModulation/Source/AudioModulationEditor/Private/Layouts/SoundModulationParameterSettingsLayout.cpp b/Engine/Plugins/Runtime/AudioModulation/Source/AudioModulationEditor/Private/Layouts/SoundModulationParameterSettingsLayout.cpp
--- a/Engine/Plugins/Runtime/AudioModulation/Source/AudioModulationEditor/Private/Layouts/SoundModulationParameterSettingsLayout.cpp
+++ b/Engine/Plugins/Runtime/AudioModulation/Source/AudioModulationEditor/Private/Layouts/SoundModulationParameterSettingsLayout.cpp
@@ -62,7 +62,7 @@ void FSoundModulationParameterSettingsLayoutCustomization::CustomizeChildren(TSh
 		PropertyHandles.Add(PropertyName, ChildHandle);
 	}
 
-	TSharedRef<IPropertyHandle> LinearValueHandle = PropertyHandles.FindChecked(GET_MEMBER_NAME_CHECKED(FSoundModulationParameterSettings, ValueLinear)).ToSharedRef();
+	TSharedRef<IPropertyHandle> NormalizedValueHandle = PropertyHandles.FindChecked(GET_MEMBER_NAME_CHECKED(FSoundModulationParameterSettings, ValueNormalized)).ToSharedRef();
 	TSharedRef<IPropertyHandle> UnitValueHandle = PropertyHandles.FindChecked(GET_MEMBER_NAME_CHECKED(FSoundModulationParameterSettings, ValueUnit)).ToSharedRef();
 	TSharedRef<IPropertyHandle> UnitDisplayNameHandle = PropertyHandles.FindChecked(GET_MEMBER_NAME_CHECKED(FSoundModulationParameterSettings, UnitDisplayName)).ToSharedRef();
 
@@ -79,12 +79,12 @@ void FSoundModulationParameterSettingsLayoutCustomization::CustomizeChildren(TSh
 		{
 			if (USoundModulationParameter* Parameter = ModParamSettingsUtils::GetStructParameter(StructPropertyHandle))
 			{
-				Parameter->RefreshLinearValue();
+				Parameter->RefreshNormalizedValue();
 				Parameter->RefreshUnitValue();
 			}
 		}));
 
-		LinearValueHandle->SetOnPropertyValueChanged(FSimpleDelegate::CreateLambda([StructPropertyHandle]()
+		NormalizedValueHandle->SetOnPropertyValueChanged(FSimpleDelegate::CreateLambda([StructPropertyHandle]()
 		{
 			if (USoundModulationParameter* Parameter = ModParamSettingsUtils::GetStructParameter(StructPropertyHandle))
 			{
@@ -136,7 +136,7 @@ void FSoundModulationParameterSettingsLayoutCustomization::CustomizeChildren(TSh
 				.Padding(4.0f, 0.0f, 0.0f, 0.0f)
 				.VAlign(VAlign_Center)
 				[
-					LinearValueHandle->CreatePropertyValueWidget()
+					NormalizedValueHandle->CreatePropertyValueWidget()
 				]
 			+ SHorizontalBox::Slot()
 				.Padding(4.0f, 0.0f, 0.0f, 0.0f)
@@ -145,7 +145,7 @@ void FSoundModulationParameterSettingsLayoutCustomization::CustomizeChildren(TSh
 				[
 					SNew(STextBlock)
 					.Font(IDetailLayoutBuilder::GetDetailFont())
-					.Text(LOCTEXT("SoundModulationControl_UnitValueLinear", "Linear"))
+					.Text(LOCTEXT("SoundModulationControl_UnitValueNormalized", "Normalized"))
 					.IsEnabled(false)
 				]
 		];
@@ -164,7 +164,7 @@ void FSoundModulationParameterSettingsLayoutCustomization::CustomizeChildren(TSh
 			]
 			.ValueContent()
 			[
-				LinearValueHandle->CreatePropertyValueWidget()
+				NormalizedValueHandle->CreatePropertyValueWidget()
 			];
 	}
 
diff --git a/Engine/Source/Runtime/Engine/Classes/Sound/SoundModulationDestination.h b/Engine/Source/Runtime/Engine/Classes/Sound/SoundModulationDestination.h
--- a/Engine/Source/Runtime/Engine/Classes/Sound/SoundModulationDestination.h
+++ b/Engine/Source/Runtime/Engine/Classes/Sound/SoundModulationDestination.h
@@ -30,7 +30,7 @@ struct ENGINE_API FSoundModulationDestinationSettings
 	bool bEnableModulation = false;
 #endif // WITH_EDITORONLY_DATA
 
-	/** Subscribed modulator to listen to apply result to base value. */
+	/** Modulation source, which provides value to mix with base value. */
 	UPROPERTY(EditAnywhere, BlueprintReadWrite, Category = Modulation)
 	USoundModulatorBase* Modulator = nullptr;
 };
@@ -50,17 +50,17 @@ namespace Audio
 			/** Initializes the modulation destination
 			 * InDeviceId - DeviceId associated with modulation plugin instance
 			 * bInIsBuffered - Whether or not to run destination in "buffered mode," which manages an internal buffer to smooth modulation value between process calls
-			 * bInValueLinear - Whether or not to keep the output value in linear [0.0f, 1.0f] space
+			 * bInValueNormalized - Whether or not to keep the output value in normalized, unitless [0.0f, 1.0f] space
 			 */
-			void Init(FDeviceId InDeviceId, bool bInIsBuffered = false, bool bInValueLinear = false);
+			void Init(FDeviceId InDeviceId, bool bInIsBuffered = false, bool bInValueNormalized = false);
 
 			/** Initializes the modulation destination
 			 * InDeviceId - DeviceId associated with modulation plugin instance
 			 * InParameterName - Name of parameter used to mix/convert destination value to/from unit space
 			 * bInIsBuffered - Whether or not to run destination in "buffered mode," which manages an internal buffer to smooth modulation value between process calls
-			 * bInValueLinear - Whether or not to keep the output value in linear [0.0f, 1.0f] space
+			 * bInValueNormalized - Whether or not to keep the output value in normalized, unitless [0.0f, 1.0f] space
 			 */
-			void Init(FDeviceId InDeviceId, FName InParameterName, bool bInIsBuffered = false, bool bInValueLinear = false);
+			void Init(FDeviceId InDeviceId, FName InParameterName, bool bInIsBuffered = false, bool bInValueNormalized = false);
 
 			/** returns whether or not destination references an active modulator */
 			bool IsActive();
@@ -81,12 +81,12 @@ namespace Audio
 			float ValueTarget = 1.0f;
 
 			uint8 bIsBuffered   = 0;
-			uint8 bValueLinear  = 0;
+			uint8 bValueNormalized  = 0;
 			uint8 bIsActive     = 0;
 			uint8 bHasProcessed = 0;
 
-			AlignedFloatBuffer OutputBuffer;
-			AlignedFloatBuffer TempBufferLinear;
+			FAlignedFloatBuffer OutputBuffer;
+			FAlignedFloatBuffer TempBufferNormalized;
 			FModulatorHandle Handle;
 
 			FName ParameterName;
@@ -96,7 +96,7 @@ namespace Audio
 
 		public:
 			/** Returns buffer of interpolated modulation values */
-			FORCEINLINE const AlignedFloatBuffer& GetBuffer() const
+			FORCEINLINE const FAlignedFloatBuffer& GetBuffer() const
 			{
 				check(bIsBuffered);
 				return OutputBuffer;
@@ -110,7 +110,7 @@ namespace Audio
 			}
 
 			/** Returns sample value last reported by modulator. Returns value in unit space, unless 
-			 * 'ValueLinear' option is set on initialization.
+			 * 'ValueNormalized' option is set on initialization.
 			 */
 			FORCEINLINE float GetValue() const
 			{
diff --git a/Engine/Source/Runtime/Engine/Private/SoundModulationDestination.cpp b/Engine/Source/Runtime/Engine/Private/SoundModulationDestination.cpp
--- a/Engine/Source/Runtime/Engine/Private/SoundModulationDestination.cpp
+++ b/Engine/Source/Runtime/Engine/Private/SoundModulationDestination.cpp
@@ -14,9 +14,9 @@ namespace Audio
 		: DeviceId(InModulationDestination.DeviceId)
 		, ValueTarget(InModulationDestination.ValueTarget)
 		, bIsBuffered(InModulationDestination.bIsBuffered)
-		, bValueLinear(InModulationDestination.bValueLinear)
+		, bValueNormalized(InModulationDestination.bValueNormalized)
 		, OutputBuffer(InModulationDestination.OutputBuffer)
-		, TempBufferLinear(InModulationDestination.TempBufferLinear)
+		, TempBufferNormalized(InModulationDestination.TempBufferNormalized)
 		, Handle(InModulationDestination.Handle)
 		, ParameterName(InModulationDestination.ParameterName)
 		, Parameter(InModulationDestination.Parameter)
@@ -28,9 +28,9 @@ namespace Audio
 		DeviceId			= InModulationDestination.DeviceId;
 		ValueTarget			= InModulationDestination.ValueTarget;
 		bIsBuffered			= InModulationDestination.bIsBuffered;
-		bValueLinear		= InModulationDestination.bValueLinear;
+		bValueNormalized		= InModulationDestination.bValueNormalized;
 		OutputBuffer		= InModulationDestination.OutputBuffer;
-		TempBufferLinear	= InModulationDestination.TempBufferLinear;
+		TempBufferNormalized	= InModulationDestination.TempBufferNormalized;
 		Handle				= InModulationDestination.Handle;
 		ParameterName		= InModulationDestination.ParameterName;
 		Parameter			= InModulationDestination.Parameter;
@@ -43,9 +43,9 @@ namespace Audio
 		DeviceId			= MoveTemp(InModulationDestination.DeviceId);
 		ValueTarget			= MoveTemp(InModulationDestination.ValueTarget);
 		bIsBuffered			= MoveTemp(InModulationDestination.bIsBuffered);
-		bValueLinear		= MoveTemp(InModulationDestination.bValueLinear);
+		bValueNormalized		= MoveTemp(InModulationDestination.bValueNormalized);
 		OutputBuffer		= MoveTemp(InModulationDestination.OutputBuffer);
-		TempBufferLinear	= MoveTemp(InModulationDestination.TempBufferLinear);
+		TempBufferNormalized	= MoveTemp(InModulationDestination.TempBufferNormalized);
 		Handle				= MoveTemp(InModulationDestination.Handle);
 		ParameterName		= MoveTemp(InModulationDestination.ParameterName);
 		Parameter			= MoveTemp(InModulationDestination.Parameter);
@@ -55,14 +55,14 @@ namespace Audio
 		return *this;
 	}
 
-	void FModulationDestination::Init(FDeviceId InDeviceId, bool bInIsBuffered, bool bInValueLinear)
+	void FModulationDestination::Init(FDeviceId InDeviceId, bool bInIsBuffered, bool bInValueNormalized)
 	{
 		DeviceId = InDeviceId;
 		bIsBuffered = bInIsBuffered;
-		bValueLinear = bInValueLinear;
+		bValueNormalized = bInValueNormalized;
 
 		OutputBuffer.Reset();
-		TempBufferLinear.Reset();
+		TempBufferNormalized.Reset();
 		ParameterName = FName();
 
 		FScopeLock Lock(&SettingsCritSection);
@@ -79,9 +79,9 @@ namespace Audio
 		return bIsActive > 0;
 	}
 
-	void FModulationDestination::Init(FDeviceId InDeviceId, FName InParameterName, bool bInIsBuffered, bool bInValueLinear)
+	void FModulationDestination::Init(FDeviceId InDeviceId, FName InParameterName, bool bInIsBuffered, bool bInValueNormalized)
 	{
-		Init(InDeviceId, bInIsBuffered, bInValueLinear);
+		Init(InDeviceId, bInIsBuffered, bInValueNormalized);
 		ParameterName = InParameterName;
 	}
 
@@ -91,11 +91,11 @@ namespace Audio
 
 		bHasProcessed = 1;
 		float LastTarget = ValueTarget;
-		float NewTargetLinear = Parameter.DefaultValue;
+		float NewTargetNormalized = Parameter.DefaultValue;
 
 		if (Parameter.bRequiresConversion)
 		{
-			Parameter.LinearFunction(&NewTargetLinear, 1);
+			Parameter.NormalizedFunction(&NewTargetNormalized, 1);
 		}
 
 		FScopeLock Lock(&SettingsCritSection);
@@ -103,10 +103,10 @@ namespace Audio
 			bIsActive = Handle.IsValid();
 			if (bIsActive)
 			{
-				Handle.GetValue(NewTargetLinear);
+				Handle.GetValue(NewTargetNormalized);
 			}
 		}
-		ValueTarget = NewTargetLinear;
+		ValueTarget = NewTargetNormalized;
 
 		if (OutputBuffer.Num() != InNumSamples)
 		{
@@ -116,25 +116,25 @@ namespace Audio
 		BufferSetToConstantInplace(OutputBuffer, 1.0f);
 		FadeBufferFast(OutputBuffer, LastTarget, ValueTarget);
 
-		if (TempBufferLinear.Num() != InNumSamples)
+		if (TempBufferNormalized.Num() != InNumSamples)
 		{
-			TempBufferLinear.Reset();
-			TempBufferLinear.AddUninitialized(InNumSamples);
+			TempBufferNormalized.Reset();
+			TempBufferNormalized.AddUninitialized(InNumSamples);
 		}
 
-		FMemory::Memcpy(TempBufferLinear.GetData(), InBufferUnitBase, sizeof(float) * InNumSamples);
+		FMemory::Memcpy(TempBufferNormalized.GetData(), InBufferUnitBase, sizeof(float) * InNumSamples);
 
 		// Convert input buffer to linear space if necessary
 		if (Parameter.bRequiresConversion)
 		{
-			Parameter.LinearFunction(TempBufferLinear.GetData(), TempBufferLinear.Num());
+			Parameter.NormalizedFunction(TempBufferNormalized.GetData(), TempBufferNormalized.Num());
 		}
 
 		// Mix mod value and base value buffers in linear space
-		Parameter.MixFunction(OutputBuffer.GetData(), TempBufferLinear.GetData(), InNumSamples);
+		Parameter.MixFunction(OutputBuffer.GetData(), TempBufferNormalized.GetData(), InNumSamples);
 
 		// Convert result to unit space if necessary
-		if (Parameter.bRequiresConversion && !bValueLinear)
+		if (Parameter.bRequiresConversion && !bValueNormalized)
 		{
 			Parameter.UnitFunction(OutputBuffer.GetData(), OutputBuffer.Num());
 		}
@@ -144,11 +144,11 @@ namespace Audio
 	{
 		bHasProcessed = 1;
 		float LastTarget = ValueTarget;
-		float NewTargetLinear = Parameter.DefaultValue;
+		float NewTargetNormalized = Parameter.DefaultValue;
 
 		if (Parameter.bRequiresConversion)
 		{
-			Parameter.LinearFunction(&NewTargetLinear, 1);
+			Parameter.NormalizedFunction(&NewTargetNormalized, 1);
 		}
 
 		FScopeLock Lock(&SettingsCritSection);
@@ -156,15 +156,15 @@ namespace Audio
 			bIsActive = Handle.IsValid();
 			if (bIsActive)
 			{
-				Handle.GetValue(NewTargetLinear);
+				Handle.GetValue(NewTargetNormalized);
 			}
 		}
 
 		// Convert base to linear space
-		float InValueBaseLinear = InValueUnitBase;
+		float InValueBaseNormalized = InValueUnitBase;
 		if (Parameter.bRequiresConversion)
 		{
-			Parameter.LinearFunction(&InValueBaseLinear, 1);
+			Parameter.NormalizedFunction(&InValueBaseNormalized, 1);
 		}
 
 		if (bIsBuffered)
@@ -177,11 +177,11 @@ namespace Audio
 		}
 
 		// Mix in base value
-		Parameter.MixFunction(&NewTargetLinear, &InValueBaseLinear, 1);
-		ValueTarget = NewTargetLinear;
+		Parameter.MixFunction(&NewTargetNormalized, &InValueBaseNormalized, 1);
+		ValueTarget = NewTargetNormalized;
 
 		// Convert target to unit space if required
-		if (Parameter.bRequiresConversion && !bValueLinear)
+		if (Parameter.bRequiresConversion && !bValueNormalized)
 		{
 			Parameter.UnitFunction(&ValueTarget, 1);
 		}
