Missing WITH_EDITORONLY_DATA macro

diff --git a/Engine/Plugins/Runtime/AudioModulation/Source/AudioModulation/Private/AudioModulationSystem.cpp b/Engine/Plugins/Runtime/AudioModulation/Source/AudioModulation/Private/AudioModulationSystem.cpp
--- a/Engine/Plugins/Runtime/AudioModulation/Source/AudioModulation/Private/AudioModulationSystem.cpp
+++ b/Engine/Plugins/Runtime/AudioModulation/Source/AudioModulation/Private/AudioModulationSystem.cpp
@@ -274,12 +274,16 @@ namespace AudioModulation
 						Parameter.ParameterName = InParamName;
 						Parameter.bRequiresConversion = Param->RequiresUnitConversion();
 						Parameter.MixFunction = Param->GetMixFunction();
-						Parameter.UnitDisplayName = Param->Settings.UnitDisplayName;
 						Parameter.UnitFunction = Param->GetUnitConversionFunction();
 						Parameter.NormalizedFunction = Param->GetNormalizedConversionFunction();
 						Parameter.DefaultValue = Param->GetUnitDefault();
 						Parameter.MinValue = Param->GetUnitMin();
 						Parameter.MaxValue = Param->GetUnitMax();
+
+#if WITH_EDITORONLY_DATA
+						Parameter.UnitDisplayName = Param->Settings.UnitDisplayName;
+#endif // WITH_EDITORONLY_DATA
+
 						return Parameter;
 					}
 				}
