Edigrate from 4.26: Fix missing macro for DebugName in shipped builds

diff --git a/Engine/Plugins/Runtime/AudioModulation/Source/AudioModulation/Private/Generators/SoundModulationEnvelopeFollower.cpp b/Engine/Plugins/Runtime/AudioModulation/Source/AudioModulation/Private/Generators/SoundModulationEnvelopeFollower.cpp
--- a/Engine/Plugins/Runtime/AudioModulation/Source/AudioModulation/Private/Generators/SoundModulationEnvelopeFollower.cpp
+++ b/Engine/Plugins/Runtime/AudioModulation/Source/AudioModulation/Private/Generators/SoundModulationEnvelopeFollower.cpp
@@ -9,7 +9,9 @@
 
 namespace AudioModulation
 {
+#if !UE_BUILD_SHIPPING
 	const FString FEnvelopeFollowerGenerator::DebugName = TEXT("EnvelopeFollower");
+#endif // !UE_BUILD_SHIPPING
 
 	FEnvelopeFollowerGenerator::FEnvelopeFollowerGenerator(const FEnvelopeFollowerGeneratorParams& InParams, Audio::FDeviceId InDeviceId)
 		: IGenerator(InDeviceId)
