Edigrate following from 4 to 5:

Hide Modulation properly when disabled from SoundBase
    - Move canonical modulation destinations on SoundBase to struct and making a simple customized display that looks up whether modulation is enabled
    - Remove dead code/bModulationEnabled property from SoundBase from earlier attempt to use UObject reflection to hide modulation parameters on sound assets

Fix Modulation Debugger to use new au. commands properly
- Fix bugs around general display of generator debug data

Fix missing array entry manipulators on Modulation Mix Stage property

- Clean-up Generator property UX
- Add invert flag to EnvelopeFollower to allow AudioBus to gate source

- Fix missing AudioDeviceId assignment in AudioModulationSystem (bad merge?)

diff --git a/Engine/Plugins/Runtime/AudioModulation/Source/AudioModulation/Private/AudioModulation.cpp b/Engine/Plugins/Runtime/AudioModulation/Source/AudioModulation/Private/AudioModulation.cpp
--- a/Engine/Plugins/Runtime/AudioModulation/Source/AudioModulation/Private/AudioModulation.cpp
+++ b/Engine/Plugins/Runtime/AudioModulation/Source/AudioModulation/Private/AudioModulation.cpp
@@ -77,6 +77,39 @@ namespace AudioModulation
 		ModSystem->DeactivateGenerator(InGenerator);
 	}
 
+#if !UE_BUILD_SHIPPING
+	void FAudioModulation::SetDebugBusFilter(const FString* InNameFilter)
+	{
+		ModSystem->SetDebugBusFilter(InNameFilter);
+	}
+
+	void FAudioModulation::SetDebugGeneratorFilter(const FString* InFilter)
+	{
+		ModSystem->SetDebugGeneratorFilter(InFilter);
+	}
+
+	void FAudioModulation::SetDebugGeneratorTypeFilter(const FString* InFilter, bool bInIsEnabled)
+	{
+		ModSystem->SetDebugGeneratorTypeFilter(InFilter, bInIsEnabled);
+	}
+
+	void FAudioModulation::SetDebugGeneratorsEnabled(bool bInIsEnabled)
+	{
+		ModSystem->SetDebugGeneratorsEnabled(bInIsEnabled);
+	}
+
+	void FAudioModulation::SetDebugMatrixEnabled(bool bInIsEnabled)
+	{
+		ModSystem->SetDebugMatrixEnabled(bInIsEnabled);
+	}
+
+	void FAudioModulation::SetDebugMixFilter(const FString* InNameFilter)
+	{
+		ModSystem->SetDebugMixFilter(InNameFilter);
+	}
+
+#endif // !UE_BUILD_SHIPPING
+
 	void FAudioModulation::SaveMixToProfile(const USoundControlBusMix& InBusMix, const int32 InProfileIndex)
 	{
 		ModSystem->SaveMixToProfile(InBusMix, InProfileIndex);
diff --git a/Engine/Plugins/Runtime/AudioModulation/Source/AudioModulation/Private/AudioModulationDebugger.cpp b/Engine/Plugins/Runtime/AudioModulation/Source/AudioModulation/Private/AudioModulationDebugger.cpp
--- a/Engine/Plugins/Runtime/AudioModulation/Source/AudioModulation/Private/AudioModulationDebugger.cpp
+++ b/Engine/Plugins/Runtime/AudioModulation/Source/AudioModulation/Private/AudioModulationDebugger.cpp
@@ -5,6 +5,7 @@
 #if !UE_BUILD_SHIPPING
 #include "CoreMinimal.h"
 #include "Async/TaskGraphInterfaces.h"
+#include "AudioModulation.h"
 #include "AudioModulationLogging.h"
 #include "AudioModulationSystem.h"
 #include "AudioThread.h"
@@ -19,21 +20,151 @@
 #include "SoundModulationValue.h"
 
 
-static float AudioModulationDebugUpdateRateCVar = 0.25f;
+namespace AudioModulation
+{
+	namespace Debug
+	{
+		using FDebugViewUpdateFunction = TUniqueFunction<void(FAudioModulation& /*InModulation*/, const FString* /*InFilter*/)>;
+		void UpdateObjectFilter(const TArray<FString>& Args, UWorld* World, FDebugViewUpdateFunction InUpdateViewFunc)
+		{
+			if (!World)
+			{
+				return;
+			}
+
+			const FString* FilterString = nullptr;
+			if (Args.Num() > 0)
+			{
+				FilterString = &Args[0];
+			}
+
+			FAudioDeviceHandle DeviceHandle = World->GetAudioDevice();
+			if (DeviceHandle.IsValid())
+			{
+				if (GEngine)
+				{
+					GEngine->Exec(World, TEXT("au.Debug.SoundModulators 1"));
+				}
+
+				if (IAudioModulation* Modulation = DeviceHandle->ModulationInterface.Get())
+				{
+					FAudioModulation& ModulationImpl = *static_cast<FAudioModulation*>(Modulation);
+					InUpdateViewFunc(ModulationImpl, FilterString);
+				}
+			}
+		}
+	}
+}
+
+static float AudioModulationDebugUpdateRateCVar = 0.1f;
 FAutoConsoleVariableRef CVarAudioModulationDebugUpdateRate(
 	TEXT("au.Debug.SoundModulators.UpdateRate"),
 	AudioModulationDebugUpdateRateCVar,
 	TEXT("Sets update rate for modulation debug statistics (in seconds).\n")
-	TEXT("Default: 0.25f"),
+	TEXT("Default: 0.1f"),
 	ECVF_Default);
 
+static FAutoConsoleCommandWithWorldAndArgs CVarAudioModulationDebugBuses(
+	TEXT("au.Debug.SoundModulators.Filter.Buses"),
+	TEXT("Sets substring by which to filter mixes in matrix view. Arguments:\n"
+		"Enabled (Optional, ex. True, False. Default: True) - Whether or not to enable showing buses.\n"
+		"Filter (Optional, Default: null) - Whether or not to filter buses by name using the provided substring.\n"),
+	FConsoleCommandWithWorldAndArgsDelegate::CreateLambda([](const TArray<FString>& Args, UWorld* World)
+	{
+		using namespace AudioModulation;
+		Debug::UpdateObjectFilter(Args, World, [](FAudioModulation& InModulation, const FString* InFilter)
+		{
+			InModulation.SetDebugBusFilter(InFilter);
+		});
+	}));
+
+static FAutoConsoleCommandWithWorldAndArgs CVarAudioModulationDebugMixes(
+	TEXT("au.Debug.SoundModulators.Filter.Mixes"),
+	TEXT("Sets substring by which to filter mixes in matrix view. Arguments:\n"
+		"Filter - Filter bus mixes by name using the provided substring.\n"),
+	FConsoleCommandWithWorldAndArgsDelegate::CreateLambda([](const TArray<FString>& Args, UWorld* World)
+	{
+		using namespace AudioModulation;
+		Debug::UpdateObjectFilter(Args, World, [](FAudioModulation& InModulation, const FString* InFilter)
+		{
+			InModulation.SetDebugMixFilter(InFilter);
+		});
+	}));
+
+static FAutoConsoleCommandWithWorldAndArgs CVarAudioModulationDebugMatrix(
+	TEXT("au.Debug.SoundModulators.Enable.Matrix"),
+	TEXT("Whether or not to enable mix matrix. Arguments:\n"
+		"Enabled (Default: 1/True) - Whether or not to have matrix view enabled.\n"),
+	FConsoleCommandWithWorldAndArgsDelegate::CreateLambda([](const TArray<FString>& Args, UWorld* World)
+	{
+		using namespace AudioModulation;
+		Debug::UpdateObjectFilter(Args, World, [](FAudioModulation& InModulation, const FString* InFilter)
+		{
+			if (InFilter)
+			{
+				const bool bEnabled = InFilter->ToBool();
+				InModulation.SetDebugMatrixEnabled(bEnabled);
+			}
+		});
+	}));
+
+static FAutoConsoleCommandWithWorldAndArgs CVarAudioModulationSetDebugGeneratorFilter(
+	TEXT("au.Debug.SoundModulators.Filter.Generators"),
+	TEXT("Sets substring by which to filter generators. Arguments:\n"
+		"Name - Filter generators by name using the provided substring.\n"),
+	FConsoleCommandWithWorldAndArgsDelegate::CreateLambda([](const TArray<FString>& Args, UWorld* World)
+	{
+		using namespace AudioModulation;
+
+		Debug::UpdateObjectFilter(Args, World, [](FAudioModulation& InModulation, const FString* InFilter)
+		{
+			InModulation.SetDebugGeneratorFilter(InFilter);
+		});
+	}));
+
+static FAutoConsoleCommandWithWorldAndArgs CVarAudioModulationSetDebugGeneratorFilterType(
+	TEXT("au.Debug.SoundModulators.Filter.Generators.Type"),
+	TEXT("Whether to display or hide Generator type provided (defaults to show if enablement boolean not provided). Arguments:\n"
+		"Name (Optional, Default: null) - Filter generators type to display/hide (Empty/null clears any currently set filter).\n"
+		"Enabled (Optional, Default: True) - True to show, false to hide.\n"),
+	FConsoleCommandWithWorldAndArgsDelegate::CreateLambda([](const TArray<FString>& Args, UWorld* World)
+	{
+		using namespace AudioModulation;
+
+		bool bEnabled = true;
+		if (Args.Num() > 1)
+		{
+			bEnabled = Args[1].ToBool();
+		}
+		Debug::UpdateObjectFilter(Args, World, [bEnabled](FAudioModulation& InModulation, const FString* InFilter)
+		{
+			InModulation.SetDebugGeneratorTypeFilter(InFilter, bEnabled);
+		});
+	}));
+
+static FAutoConsoleCommandWithWorldAndArgs CVarAudioModulationDebugEnableGenerators(
+	TEXT("au.Debug.SoundModulators.Enable.Generators"),
+	TEXT("Whether or not to enable displaying generators. Arguments:\n"
+		"Enabled (Default: 1/True) - Whether or not to have generators view enabled.\n"),
+	FConsoleCommandWithWorldAndArgsDelegate::CreateLambda([](const TArray<FString>& Args, UWorld* World)
+	{
+		using namespace AudioModulation;
+		Debug::UpdateObjectFilter(Args, World, [](FAudioModulation& InModulation, const FString* InFilter)
+		{
+			if (InFilter)
+			{
+				const bool bEnabled = InFilter->ToBool();
+				InModulation.SetDebugGeneratorsEnabled(bEnabled);
+			}
+		});
+	}));
 
 namespace AudioModulation
 {
 	namespace Debug
 	{
 		const int32 MaxNameLength = 40;
-		const int32 XIndent       = 36;
+		const int32 XIndent = 36;
 
 		FColor GetUnitRenderColor(const float Value)
 		{
@@ -42,40 +173,21 @@ namespace AudioModulation
 				: FColor::Green;
 		}
 
-		int32 RenderStatGenerators(const FGeneratorDebugInfo& GeneratorDebugInfo, FCanvas& Canvas, int32 X, int32 Y, const UFont& Font)
+		int32 RenderStatGenerators(const FString& Name, const FGeneratorDebugInfo& GeneratorDebugInfo, int32 CellWidth, FCanvas& Canvas, int32 X, int32 Y, const UFont& Font)
 		{
 			int32 Height = 12;
 			int32 Width = 12;
 			Font.GetStringHeightAndWidth(TEXT("@"), Height, Width);
 
-			int32 CellWidth = 0;
-			for (const FString& Category : GeneratorDebugInfo.Categories)
-			{
-				CellWidth = FMath::Max(CellWidth, Category.Len());
-			}
-
-			for (const FGeneratorDebugInfo::FInstanceValues& InstanceInfo : GeneratorDebugInfo.FilteredInstances)
-			{
-				for (const FString& Value : InstanceInfo)
-				{
-					CellWidth = FMath::Max(CellWidth, Value.Len());
-				}
-			}
-
-			Canvas.DrawShadowedString(X + XIndent, Y, TEXT("Active Generators:"), &Font, FColor::Red);
 			Y += Height;
 			X += XIndent;
 
+			Canvas.DrawShadowedString(X, Y, *(Name + TEXT(":")), &Font, FColor::Red);
+			Y += Height;
+
 			// Print Category Titles
 			{
 				int32 RowX = X;
-
-				Canvas.DrawShadowedString(RowX, Y, TEXT("Name"), &Font, FColor::White);
-				RowX += Width * CellWidth;
-
-				Canvas.DrawShadowedString(RowX, Y, TEXT("Ref Count"), &Font, FColor::White);
-				RowX += Width * CellWidth;
-
 				for (const FString& Category : GeneratorDebugInfo.Categories)
 				{
 					Canvas.DrawShadowedString(RowX, Y, *Category, &Font, FColor::White);
@@ -90,8 +202,7 @@ namespace AudioModulation
 				int32 RowX = X;
 
 				// Validate that all instances have the same number of values as there are registered categories.
-				// 2 "extra" are from parent generator proxy that include proxy's Name & Ref Count.
-				ensure(InstanceValues.Num() == GeneratorDebugInfo.Categories.Num() + 2);
+				ensure(InstanceValues.Num() == GeneratorDebugInfo.Categories.Num());
 
 				for (const FString& Value : InstanceValues)
 				{
@@ -105,7 +216,9 @@ namespace AudioModulation
 					RowX += Width * CellWidth;
 				}
 			}
-			Y += Height;
+
+			// Add space after, so twice the height
+			Y += Height * 2;
 
 			return Y;
 		}
@@ -281,6 +394,8 @@ namespace AudioModulation
 	FAudioModulationDebugger::FAudioModulationDebugger()
 		: bActive(0)
 		, bShowRenderStatMix(1)
+		, bShowGenerators(0)
+		, bEnableAllGenerators(0)
 		, ElapsedSinceLastUpdate(0.0f)
 	{
 	}
@@ -344,7 +459,9 @@ namespace AudioModulation
 		for (const FModulatorGeneratorProxy* Proxy : FilteredGeneratorProxies)
 		{
 			FGeneratorDebugInfo::FInstanceValues Values = Proxy->GetDebugValues();
-			RefreshedFilteredGenerators.FindOrAdd(Proxy->GetDebugName()).FilteredInstances.Emplace(MoveTemp(Values));
+			FGeneratorDebugInfo& DebugInfo = RefreshedFilteredGenerators.FindOrAdd(Proxy->GetDebugName());
+			DebugInfo.FilteredInstances.Emplace(MoveTemp(Values));
+			DebugInfo.Categories = Proxy->GetDebugCategories();
 		}
 
 		FFunctionGraphTask::CreateAndDispatchWhenReady([this, RefreshedFilteredBuses, RefreshedFilteredMixes, RefreshedFilteredGenerators]()
@@ -355,17 +472,28 @@ namespace AudioModulation
 			FilteredMixes = RefreshedFilteredMixes;
 			FilteredMixes.Sort(&Debug::CompareNames<FControlBusMixDebugInfo>);
 
-			FGeneratorSortMap NewFilteredMap = RefreshedFilteredGenerators;
-			for (TPair<FString, FGeneratorDebugInfo>& NewInfoPair : NewFilteredMap)
+			FGeneratorSortMap NewFilteredMap;
+			if (bShowGenerators)
 			{
-				if (FGeneratorDebugInfo* LastDebugInfo = FilteredGeneratorsMap.Find(NewInfoPair.Key))
+				NewFilteredMap = RefreshedFilteredGenerators;
+				for (TPair<FString, FGeneratorDebugInfo>& NewInfoPair : NewFilteredMap)
 				{
-					NewInfoPair.Value.bEnabled = LastDebugInfo->bEnabled;
+					FGeneratorDebugInfo* LastDebugInfo = FilteredGeneratorsMap.Find(NewInfoPair.Key);
+					NewInfoPair.Value.bEnabled = bEnableAllGenerators
+						|| (LastDebugInfo && LastDebugInfo->bEnabled);
+
+					if (bool* Value = RequestedGeneratorUpdate.Find(NewInfoPair.Key))
+					{
+						NewInfoPair.Value.bEnabled = *Value;
+					}
+
 					if (!NewInfoPair.Value.bEnabled)
 					{
 						NewInfoPair.Value.FilteredInstances.Reset();
 					}
 				}
+				bEnableAllGenerators = 0;
+				RequestedGeneratorUpdate.Reset();
 			}
 
 			FilteredGeneratorsMap = MoveTemp(NewFilteredMap);
@@ -388,6 +516,81 @@ namespace AudioModulation
 		return true;
 	}
 
+	void FAudioModulationDebugger::SetDebugBusFilter(const FString* InNameFilter)
+	{
+		check(IsInGameThread());
+
+		if (InNameFilter)
+		{
+			BusStringFilter = *InNameFilter;
+		}
+		else
+		{
+			BusStringFilter.Reset();
+		}
+	}
+
+	void FAudioModulationDebugger::SetDebugMatrixEnabled(bool bInIsEnabled)
+	{
+		bShowRenderStatMix = bInIsEnabled;
+	}
+
+	void FAudioModulationDebugger::SetDebugMixFilter(const FString* InNameFilter)
+	{
+		check(IsInGameThread());
+
+		if (InNameFilter)
+		{
+			MixStringFilter = *InNameFilter;
+		}
+		else
+		{
+			MixStringFilter.Reset();
+		}
+	}
+
+	void FAudioModulationDebugger::SetDebugGeneratorsEnabled(bool bInIsEnabled)
+	{
+		check(IsInGameThread());
+
+		bShowGenerators = bInIsEnabled;
+		bEnableAllGenerators = 1;
+	}
+
+	void FAudioModulationDebugger::SetDebugGeneratorFilter(const FString* InFilter)
+	{
+		check(IsInGameThread());
+
+		if (InFilter)
+		{
+			bShowGenerators = 1;
+			GeneratorStringFilter = *InFilter;
+		}
+		else
+		{
+			GeneratorStringFilter.Reset();
+		}
+	}
+
+	void FAudioModulationDebugger::SetDebugGeneratorTypeFilter(const FString* InFilter, bool bInIsEnabled)
+	{
+		check(IsInGameThread());
+
+		if (InFilter)
+		{
+			bShowGenerators |= bInIsEnabled;
+			RequestedGeneratorUpdate.Add(*InFilter, bInIsEnabled);
+		}
+		else
+		{
+			bShowGenerators = bInIsEnabled;
+			if (bInIsEnabled)
+			{
+				bEnableAllGenerators = 1;
+			}
+		}
+	}
+
 	int32 FAudioModulationDebugger::OnRenderStat(FCanvas& Canvas, int32 X, int32 Y, const UFont& Font)
 	{
 		check(IsInGameThread());
@@ -404,11 +607,31 @@ namespace AudioModulation
 			Y = Debug::RenderStatMixMatrix(FilteredMixes, FilteredBuses, Canvas, X, Y, Font);
 		}
 
+		int32 CellWidth = 0;
 		for (const TPair<FString, FGeneratorDebugInfo>& GeneratorStatPair : FilteredGeneratorsMap)
 		{
-			if (GeneratorStatPair.Value.bEnabled)
+			for (const FString& Category : GeneratorStatPair.Value.Categories)
 			{
-				Y = Debug::RenderStatGenerators(GeneratorStatPair.Value, Canvas, X, Y, Font);
+				CellWidth = FMath::Max(CellWidth, Category.Len());
+			}
+
+			for (const FGeneratorDebugInfo::FInstanceValues& InstanceInfo : GeneratorStatPair.Value.FilteredInstances)
+			{
+				for (const FString& Value : InstanceInfo)
+				{
+					CellWidth = FMath::Max(CellWidth, Value.Len());
+				}
+			}
+		}
+
+		if (bShowGenerators)
+		{
+			for (const TPair<FString, FGeneratorDebugInfo>& GeneratorStatPair : FilteredGeneratorsMap)
+			{
+				if (GeneratorStatPair.Value.bEnabled)
+				{
+					Y = Debug::RenderStatGenerators(GeneratorStatPair.Key, GeneratorStatPair.Value, CellWidth, Canvas, X, Y, Font);
+				}
 			}
 		}
 
diff --git a/Engine/Plugins/Runtime/AudioModulation/Source/AudioModulation/Private/AudioModulationDebugger.h b/Engine/Plugins/Runtime/AudioModulation/Source/AudioModulation/Private/AudioModulationDebugger.h
--- a/Engine/Plugins/Runtime/AudioModulation/Source/AudioModulation/Private/AudioModulationDebugger.h
+++ b/Engine/Plugins/Runtime/AudioModulation/Source/AudioModulation/Private/AudioModulationDebugger.h
@@ -64,6 +64,8 @@ namespace AudioModulation
 		bool bEnabled = false;
 		TArray<FString> Categories;
 
+		FString NameFilter;
+
 		using FInstanceValues = TArray<FString>;
 		TArray<FInstanceValues> FilteredInstances;
 	};
@@ -74,6 +76,12 @@ namespace AudioModulation
 		FAudioModulationDebugger();
 
 		void UpdateDebugData(double InElapsed, const FReferencedProxies& InRefProxies);
+		void SetDebugBusFilter(const FString* InNameFilter);
+		void SetDebugMatrixEnabled(bool bInIsEnabled);
+		void SetDebugMixFilter(const FString* InNameFilter);
+		void SetDebugGeneratorsEnabled(bool bInIsEnabled);
+		void SetDebugGeneratorFilter(const FString* InFilter);
+		void SetDebugGeneratorTypeFilter(const FString* InFilter, bool bInIsEnabled);
 		bool OnPostHelp(FCommonViewportClient& ViewportClient, const TCHAR* Stream);
 		int32 OnRenderStat(FCanvas& Canvas, int32 X, int32 Y, const UFont& Font);
 		bool OnToggleStat(FCommonViewportClient& ViewportClient, const TCHAR* Stream);
@@ -83,6 +91,8 @@ namespace AudioModulation
 	private:
 		uint8 bActive : 1;
 		uint8 bShowRenderStatMix : 1;
+		uint8 bShowGenerators : 1;
+		uint8 bEnableAllGenerators : 1;
 
 		TArray<FControlBusDebugInfo> FilteredBuses;
 		TArray<FControlBusMixDebugInfo> FilteredMixes;
@@ -90,6 +100,8 @@ namespace AudioModulation
 		using FGeneratorSortMap = TSortedMap<FString, FGeneratorDebugInfo, FDefaultAllocator, FGeneratorSort>;
 		FGeneratorSortMap FilteredGeneratorsMap;
 
+		TMap<FString, bool> RequestedGeneratorUpdate;
+
 		FString BusStringFilter;
 		FString GeneratorStringFilter;
 		FString MixStringFilter;
diff --git a/Engine/Plugins/Runtime/AudioModulation/Source/AudioModulation/Private/AudioModulationSystem.cpp b/Engine/Plugins/Runtime/AudioModulation/Source/AudioModulation/Private/AudioModulationSystem.cpp
--- a/Engine/Plugins/Runtime/AudioModulation/Source/AudioModulation/Private/AudioModulationSystem.cpp
+++ b/Engine/Plugins/Runtime/AudioModulation/Source/AudioModulation/Private/AudioModulationSystem.cpp
@@ -46,6 +46,8 @@ namespace AudioModulation
 
 	void FAudioModulationSystem::Initialize(const FAudioPluginInitializationParams& InitializationParams)
 	{
+		AudioDeviceId = InitializationParams.AudioDevicePtr->DeviceID;
+
 		// Load all parameters listed in settings and keep loaded.  Must be done on initialization as
 		// soft object paths cannot be loaded on non-game threads (data from parameters is copied to a proxy
 		// in 'GetParameter' via the audio thread)
@@ -193,6 +195,38 @@ namespace AudioModulation
 		});
 	}
 
+#if !UE_BUILD_SHIPPING
+	void FAudioModulationSystem::SetDebugBusFilter(const FString* InFilter)
+	{
+		Debugger.SetDebugBusFilter(InFilter);
+	}
+
+	void FAudioModulationSystem::SetDebugGeneratorsEnabled(bool bInIsEnabled)
+	{
+		Debugger.SetDebugGeneratorsEnabled(bInIsEnabled);
+	}
+
+	void FAudioModulationSystem::SetDebugGeneratorFilter(const FString* InFilter)
+	{
+		Debugger.SetDebugGeneratorFilter(InFilter);
+	}
+
+	void FAudioModulationSystem::SetDebugGeneratorTypeFilter(const FString* InFilter, bool bInEnabled)
+	{
+		Debugger.SetDebugGeneratorTypeFilter(InFilter, bInEnabled);
+	}
+
+	void FAudioModulationSystem::SetDebugMatrixEnabled(bool bInIsEnabled)
+	{
+		Debugger.SetDebugMatrixEnabled(bInIsEnabled);
+	}
+
+	void FAudioModulationSystem::SetDebugMixFilter(const FString* InNameFilter)
+	{
+		Debugger.SetDebugMixFilter(InNameFilter);
+	}
+#endif // !UE_BUILD_SHIPPING
+
 	bool FAudioModulationSystem::GetModulatorValue(const Audio::FModulatorHandle& InModulatorHandle, float& OutValue) const
 	{
 		const EModulatorType ModulatorType = static_cast<EModulatorType>(InModulatorHandle.GetTypeId());
@@ -635,7 +669,10 @@ namespace AudioModulation
 		TArray<FModulatorBusMixStageSettings> StageSettings;
 		for (const FSoundControlBusMixStage& Stage : InStages)
 		{
-			StageSettings.Emplace(Stage, AudioDeviceId);
+			if (Stage.Bus)
+			{
+				StageSettings.Emplace(Stage, AudioDeviceId);
+			}
 		}
 	
 		RunCommandOnProcessingThread([this, MixId, StageSettings, InFadeTime]()
diff --git a/Engine/Plugins/Runtime/AudioModulation/Source/AudioModulation/Private/AudioModulationSystem.h b/Engine/Plugins/Runtime/AudioModulation/Source/AudioModulation/Private/AudioModulationSystem.h
--- a/Engine/Plugins/Runtime/AudioModulation/Source/AudioModulation/Private/AudioModulationSystem.h
+++ b/Engine/Plugins/Runtime/AudioModulation/Source/AudioModulation/Private/AudioModulationSystem.h
@@ -186,6 +186,12 @@ namespace AudioModulation
 
 #if !UE_BUILD_SHIPPING
 	public:
+		void SetDebugBusFilter(const FString* InFilter);
+		void SetDebugMixFilter(const FString* InFilter);
+		void SetDebugMatrixEnabled(bool bInIsEnabled);
+		void SetDebugGeneratorsEnabled(bool bInIsEnabled);
+		void SetDebugGeneratorFilter(const FString* InFilter);
+		void SetDebugGeneratorTypeFilter(const FString* InFilter, bool bInIsEnabled);
 		bool OnPostHelp(FCommonViewportClient* ViewportClient, const TCHAR* Stream);
 		int OnRenderStat(FViewport* Viewport, FCanvas* Canvas, int32 X, int32 Y, const UFont& Font, const FVector* ViewLocation, const FRotator* ViewRotation);
 		bool OnToggleStat(FCommonViewportClient* ViewportClient, const TCHAR* Stream);
@@ -220,6 +226,12 @@ namespace AudioModulation
 		void OnAuditionEnd() { }
 
 #if !UE_BUILD_SHIPPING
+		void SetDebugBusFilter(const FString* InFilter) { }
+		void SetDebugMixFilter(const FString* InFilter) { }
+		void SetDebugMatrixEnabled(bool bInIsEnabled) { }
+		void SetDebugGeneratorsEnabled(bool bInIsEnabled) { }
+		void SetDebugGeneratorFilter(const FString* InFilter) { }
+		void SetDebugGeneratorTypeFilter(const FString* InFilter, bool bInIsEnabled) { }
 		bool OnPostHelp(FCommonViewportClient* ViewportClient, const TCHAR* Stream) { return false; }
 		int OnRenderStat(FViewport* Viewport, FCanvas* Canvas, int32 X, int32 Y, const UFont& Font, const FVector* ViewLocation, const FRotator* ViewRotation) { return Y; }
 		bool OnToggleStat(FCommonViewportClient* ViewportClient, const TCHAR* Stream) { return false; }
diff --git a/Engine/Plugins/Runtime/AudioModulation/Source/AudioModulation/Private/Generators/SoundModulationEnvelopeFollower.cpp b/Engine/Plugins/Runtime/AudioModulation/Source/AudioModulation/Private/Generators/SoundModulationEnvelopeFollower.cpp
--- a/Engine/Plugins/Runtime/AudioModulation/Source/AudioModulation/Private/Generators/SoundModulationEnvelopeFollower.cpp
+++ b/Engine/Plugins/Runtime/AudioModulation/Source/AudioModulation/Private/Generators/SoundModulationEnvelopeFollower.cpp
@@ -11,6 +11,11 @@ namespace AudioModulation
 {
 #if !UE_BUILD_SHIPPING
 	const FString FEnvelopeFollowerGenerator::DebugName = TEXT("EnvelopeFollower");
+
+	void FEnvelopeFollowerGenerator::GetDebugCategories(TArray<FString>& OutDebugCategories) const
+	{
+		OutDebugCategories = USoundModulationGeneratorEnvelopeFollower::GetDebugCategories();
+	}
 #endif // !UE_BUILD_SHIPPING
 
 	FEnvelopeFollowerGenerator::FEnvelopeFollowerGenerator(const FEnvelopeFollowerGeneratorParams& InParams, Audio::FDeviceId InDeviceId)
diff --git a/Engine/Plugins/Runtime/AudioModulation/Source/AudioModulation/Private/Generators/SoundModulationLFO.cpp b/Engine/Plugins/Runtime/AudioModulation/Source/AudioModulation/Private/Generators/SoundModulationLFO.cpp
--- a/Engine/Plugins/Runtime/AudioModulation/Source/AudioModulation/Private/Generators/SoundModulationLFO.cpp
+++ b/Engine/Plugins/Runtime/AudioModulation/Source/AudioModulation/Private/Generators/SoundModulationLFO.cpp
@@ -13,6 +13,11 @@ namespace AudioModulation
 #if !UE_BUILD_SHIPPING
 	const FString FLFOGenerator::DebugName = TEXT("LFO");
 
+	void FLFOGenerator::GetDebugCategories(TArray<FString>& OutDebugCategories) const
+	{
+		OutDebugCategories = USoundModulationGeneratorLFO::GetDebugCategories();
+	}
+
 	void FLFOGenerator::GetDebugValues(TArray<FString>& OutDebugValues) const
 	{
 		OutDebugValues.Add(FString::Printf(TEXT("%.4f"), GetValue()));
diff --git a/Engine/Plugins/Runtime/AudioModulation/Source/AudioModulation/Private/SoundModulationGeneratorProxy.h b/Engine/Plugins/Runtime/AudioModulation/Source/AudioModulation/Private/SoundModulationGeneratorProxy.h
--- a/Engine/Plugins/Runtime/AudioModulation/Source/AudioModulation/Private/SoundModulationGeneratorProxy.h
+++ b/Engine/Plugins/Runtime/AudioModulation/Source/AudioModulation/Private/SoundModulationGeneratorProxy.h
@@ -63,6 +63,19 @@ namespace AudioModulation
 		}
 
 #if !UE_BUILD_SHIPPING
+		TArray<FString> GetDebugCategories() const
+		{
+			TArray<FString> DebugCategories;
+			DebugCategories.Add(TEXT("Name"));
+			DebugCategories.Add(TEXT("Ref Count"));
+
+			TArray<FString> GeneratorCategories;
+			Generator->GetDebugCategories(GeneratorCategories);
+			DebugCategories.Append(GeneratorCategories);
+
+			return DebugCategories;
+		}
+
 		TArray<FString> GetDebugValues() const
 		{
 			TArray<FString> DebugValues;
diff --git a/Engine/Plugins/Runtime/AudioModulation/Source/AudioModulation/Private/SoundModulationProxy.h b/Engine/Plugins/Runtime/AudioModulation/Source/AudioModulation/Private/SoundModulationProxy.h
--- a/Engine/Plugins/Runtime/AudioModulation/Source/AudioModulation/Private/SoundModulationProxy.h
+++ b/Engine/Plugins/Runtime/AudioModulation/Source/AudioModulation/Private/SoundModulationProxy.h
@@ -133,7 +133,7 @@ namespace AudioModulation
 			TProxyHandle<IdType, ProxyType, ProxySettings> NewHandle = Get(InSettings.GetId(), OutProxyMap);
 			if (!NewHandle.IsValid())
 			{
-				UE_LOG(LogAudioModulation, Verbose, TEXT("Proxy created: Id '%u' for object '%s'."), NewHandle.Id, *InSettings.GetName());
+				UE_LOG(LogAudioModulation, Verbose, TEXT("Proxy created: Id '%u' for object '%s'."), ObjectId, *InSettings.GetName());
 				OutProxyMap.Add(ObjectId, ProxyType(InSettings, OutModSystem));
 				NewHandle = TProxyHandle<IdType, ProxyType, ProxySettings>(ObjectId, OutProxyMap);
 			}
diff --git a/Engine/Plugins/Runtime/AudioModulation/Source/AudioModulation/Public/AudioModulation.h b/Engine/Plugins/Runtime/AudioModulation/Source/AudioModulation/Public/AudioModulation.h
--- a/Engine/Plugins/Runtime/AudioModulation/Source/AudioModulation/Public/AudioModulation.h
+++ b/Engine/Plugins/Runtime/AudioModulation/Source/AudioModulation/Public/AudioModulation.h
@@ -56,6 +56,13 @@ namespace AudioModulation
 		void SoloBusMix(const USoundControlBusMix& InBusMix);
 
 #if !UE_BUILD_SHIPPING
+		virtual void SetDebugBusFilter(const FString* InFilter);
+		virtual void SetDebugMixFilter(const FString* InFilter);
+		virtual void SetDebugMatrixEnabled(bool bInIsEnabled);
+		virtual void SetDebugGeneratorsEnabled(bool bInIsEnabled);
+		virtual void SetDebugGeneratorFilter(const FString* InFilter);
+		virtual void SetDebugGeneratorTypeFilter(const FString* InFilter, bool bInIsEnabled);
+
 		virtual bool OnPostHelp(FCommonViewportClient* ViewportClient, const TCHAR* Stream) override;
 		virtual int32 OnRenderStat(FViewport* Viewport, FCanvas* Canvas, int32 X, int32 Y, const UFont& Font, const FVector* ViewLocation, const FRotator* ViewRotation) override;
 		virtual bool OnToggleStat(FCommonViewportClient* ViewportClient, const TCHAR* Stream) override;
diff --git a/Engine/Plugins/Runtime/AudioModulation/Source/AudioModulation/Public/AudioModulationStatics.h b/Engine/Plugins/Runtime/AudioModulation/Source/AudioModulation/Public/AudioModulationStatics.h
--- a/Engine/Plugins/Runtime/AudioModulation/Source/AudioModulation/Public/AudioModulationStatics.h
+++ b/Engine/Plugins/Runtime/AudioModulation/Source/AudioModulation/Public/AudioModulationStatics.h
@@ -51,7 +51,7 @@ public:
 		WorldContext = "WorldContextObject", 
 		Keywords = "activate modulation modulator control bus mix")
 	)
-	static void ActivateBusMix(const UObject* WorldContextObject, USoundControlBusMix* BusMix);
+	static void ActivateBusMix(const UObject* WorldContextObject, USoundControlBusMix* Mix);
 
 	/** Activates a modulation generator. Does nothing if an instance of the provided generator is already active
 	 * @param Modulator - Modulator to activate
@@ -122,7 +122,7 @@ public:
 		WorldContext = "WorldContextObject",
 		Keywords = "deactivate modulation modulator")
 	)
-	static void DeactivateBusMix(const UObject* WorldContextObject, USoundControlBusMix* BusMix);
+	static void DeactivateBusMix(const UObject* WorldContextObject, USoundControlBusMix* Mix);
 
 	/** Deactivates a modulation generator. Does nothing if an instance of the provided generator is already inactive
 	 * @param Generator - Generator to activate
@@ -144,7 +144,7 @@ public:
 		AdvancedDisplay = "2",
 		Keywords = "save serialize bus control modulation mix modulator ini")
 	)
-	static void SaveMixToProfile(const UObject* WorldContextObject, USoundControlBusMix* BusMix, int32 ProfileIndex = 0);
+	static void SaveMixToProfile(const UObject* WorldContextObject, USoundControlBusMix* Mix, int32 ProfileIndex = 0);
 
 	/** Loads control bus mix from a profile into UObject mix definition, deserialized from an ini file.
 	 * @param BusMix - Mix object to deserialize profile .ini to.
@@ -157,7 +157,7 @@ public:
 		AdvancedDisplay = "2",
 		Keywords = "load deserialize control bus modulation mix modulator ini")
 	)
-	static UPARAM(DisplayName = "Stages") TArray<FSoundControlBusMixStage> LoadMixFromProfile(const UObject* WorldContextObject, USoundControlBusMix* BusMix, bool bActivate = true, int32 ProfileIndex = 0);
+	static UPARAM(DisplayName = "Stages") TArray<FSoundControlBusMixStage> LoadMixFromProfile(const UObject* WorldContextObject, USoundControlBusMix* Mix, bool bActivate = true, int32 ProfileIndex = 0);
 
 	/** Sets a mix with the provided stage data if stages provided in active instance proxy of mix. Does not update UObject definition of mix.
 	 * @param Mix - Mix to update
@@ -170,18 +170,15 @@ public:
 		WorldContext = "WorldContextObject",
 		Keywords = "set bus control modulation modulator mix stage")
 	)
-	static void UpdateMix(const UObject* WorldContextObject, USoundControlBusMix* Mix, TArray<FSoundControlBusMixStage> Stages, float InFadeTime = -1.0f);
+	static void UpdateMix(const UObject* WorldContextObject, USoundControlBusMix* Mix, TArray<FSoundControlBusMixStage> Stages, float FadeTime = -1.0f);
 
 	/** Sets filtered stages of a given class to a provided target value for active instance of mix. Does not update UObject definition of mix.
 	 * @param Mix - Mix to modify
-	 * @param AddressFilter - Address filter to apply to provided mix's stages.
-	 * @param BusClass - Filters buses by subclass.
+	 * @param AddressFilter - (Optional) Address filter to apply to provided mix's stages.
+	 * @param ParamClassFilter - (Optional) Filters buses by parameter class.
+	 * @param ParamFilter - (Optional) Filters buses by parameter.
 	 * @param Value - Target value to mix filtered stages to.
 	 * @param FadeTime - If non-negative, updates the fade time for the resulting bus stages found matching the provided filter.
-	 * @param AttackTime - If non-negative, updates the attack time for the resulting bus stages found matching the provided filter.
-	 * @param ReleaseTime - If non-negative, updates the release time for the resulting bus stages found matching the provided filter.
-	 * @param bUpdateObject - If true, will dirty mix object and update stages on the SoundControlBusMix object in addition to updating 
-	 * the audio thread proxy.
 	 */
 	UFUNCTION(BlueprintCallable, Category = "Audio", DisplayName = "Set Control Bus Mix By Filter", meta = (
 		AdvancedDisplay = "6",
diff --git a/Engine/Plugins/Runtime/AudioModulation/Source/AudioModulation/Public/Generators/SoundModulationEnvelopeFollower.h b/Engine/Plugins/Runtime/AudioModulation/Source/AudioModulation/Public/Generators/SoundModulationEnvelopeFollower.h
--- a/Engine/Plugins/Runtime/AudioModulation/Source/AudioModulation/Public/Generators/SoundModulationEnvelopeFollower.h
+++ b/Engine/Plugins/Runtime/AudioModulation/Source/AudioModulation/Public/Generators/SoundModulationEnvelopeFollower.h
@@ -61,6 +61,7 @@ namespace AudioModulation
 #if !UE_BUILD_SHIPPING
 			static const FString DebugName;
 
+			virtual void GetDebugCategories(TArray<FString>& OutDebugCategories) const override;
 			virtual void GetDebugValues(TArray<FString>& OutDebugValues) const override;
 			virtual const FString& GetDebugName() const override;
 #endif // !UE_BUILD_SHIPPING
diff --git a/Engine/Plugins/Runtime/AudioModulation/Source/AudioModulation/Public/Generators/SoundModulationLFO.h b/Engine/Plugins/Runtime/AudioModulation/Source/AudioModulation/Public/Generators/SoundModulationLFO.h
--- a/Engine/Plugins/Runtime/AudioModulation/Source/AudioModulation/Public/Generators/SoundModulationLFO.h
+++ b/Engine/Plugins/Runtime/AudioModulation/Source/AudioModulation/Public/Generators/SoundModulationLFO.h
@@ -100,6 +100,7 @@ namespace AudioModulation
 #if !UE_BUILD_SHIPPING
 		static const FString DebugName;
 
+		virtual void GetDebugCategories(TArray<FString>& OutDebugCategories) const override;
 		virtual void GetDebugValues(TArray<FString>& OutDebugValues) const override;
 		virtual const FString& GetDebugName() const override;
 #endif // !UE_BUILD_SHIPPING
diff --git a/Engine/Plugins/Runtime/AudioModulation/Source/AudioModulation/Public/SoundModulationGenerator.h b/Engine/Plugins/Runtime/AudioModulation/Source/AudioModulation/Public/SoundModulationGenerator.h
--- a/Engine/Plugins/Runtime/AudioModulation/Source/AudioModulation/Public/SoundModulationGenerator.h
+++ b/Engine/Plugins/Runtime/AudioModulation/Source/AudioModulation/Public/SoundModulationGenerator.h
@@ -32,6 +32,8 @@ namespace AudioModulation
 #if !UE_BUILD_SHIPPING
 		virtual void GetDebugValues(TArray<FString>& OutDebugValues) const { }
 
+		virtual void GetDebugCategories(TArray<FString>& OutDebugCategories) const { }
+
 		// Required for instance look-up in factory registration
 		virtual const FString& GetDebugName() const { static const FString Invalid; return Invalid; }
 #endif // !UE_BUILD_SHIPPING
diff --git a/Engine/Plugins/Runtime/AudioModulation/Source/AudioModulationEditor/Private/Layouts/SoundControlBusMixStageLayout.cpp b/Engine/Plugins/Runtime/AudioModulation/Source/AudioModulationEditor/Private/Layouts/SoundControlBusMixStageLayout.cpp
--- a/Engine/Plugins/Runtime/AudioModulation/Source/AudioModulationEditor/Private/Layouts/SoundControlBusMixStageLayout.cpp
+++ b/Engine/Plugins/Runtime/AudioModulation/Source/AudioModulationEditor/Private/Layouts/SoundControlBusMixStageLayout.cpp
@@ -85,6 +85,32 @@ namespace AudioModulationEditorUtils
 
 void FSoundControlBusMixStageLayoutCustomization::CustomizeHeader(TSharedRef<IPropertyHandle> StructPropertyHandle, FDetailWidgetRow& HeaderRow, IPropertyTypeCustomizationUtils& CustomizationUtils)
 {
+	HeaderRow.NameContent()
+		[
+			SNew(SHorizontalBox)
+			+ SHorizontalBox::Slot()
+			.FillWidth(0.4f)
+			.Padding(2.0f, 0.0f, 0.0f, 0.0f)
+			.VAlign(VAlign_Center)
+			[
+				StructPropertyHandle->CreatePropertyNameWidget()
+			]
+		]
+		.ValueContent()
+		.MinDesiredWidth(300.0f)
+		[
+			SNew(SHorizontalBox)
+			+ SHorizontalBox::Slot()
+			.FillWidth(0.4f)
+			.Padding(2.0f, 0.0f, 0.0f, 0.0f)
+			.VAlign(VAlign_Center)
+			[
+				SNew(STextBlock)
+				.Font(IDetailLayoutBuilder::GetDetailFont())
+				.Text(FText::Format(LOCTEXT("BusMixStageHeader_Format", "Bus Stage {0}"), StructPropertyHandle->GetPropertyDisplayName()))
+				.ToolTipText(StructPropertyHandle->GetToolTipText())
+			]
+		];
 }
 
 void FSoundControlBusMixStageLayoutCustomization::CustomizeChildren(TSharedRef<IPropertyHandle> StructPropertyHandle, IDetailChildrenBuilder& ChildBuilder, IPropertyTypeCustomizationUtils& StructCustomizationUtils)
diff --git a/Engine/Source/Editor/AudioEditor/Private/AudioEditorModule.cpp b/Engine/Source/Editor/AudioEditor/Private/AudioEditorModule.cpp
--- a/Engine/Source/Editor/AudioEditor/Private/AudioEditorModule.cpp
+++ b/Engine/Source/Editor/AudioEditor/Private/AudioEditorModule.cpp
@@ -113,6 +113,8 @@ public:
 		};
 		AddCustomProperty("SoundModulationDestinationSettings", 
 			FOnGetPropertyTypeCustomizationInstance::CreateStatic(&FSoundModulationDestinationLayoutCustomization::MakeInstance));
+		AddCustomProperty("SoundModulationDefaultSettings",
+			FOnGetPropertyTypeCustomizationInstance::CreateStatic(&FSoundModulationDefaultSettingsLayoutCustomization::MakeInstance));
 
 		// Custom Class Layouts
 		auto AddCustomClass = [this, InPropertyModule = &PropertyModule](FName Name, FOnGetDetailCustomizationInstance InstanceGetter)
@@ -384,6 +386,10 @@ private:
 		SET_AUDIO_ICON_SIMPLE(SoundSubmix);
 		SET_AUDIO_ICON_SIMPLE(ReverbEffect);
 
+		SET_AUDIO_ICON(EndpointSubmix, SoundSubmix);
+		SET_AUDIO_ICON(SoundfieldEndpointSubmix, SoundSubmix);
+		SET_AUDIO_ICON(SoundfieldSubmix, SoundSubmix);
+
 		SET_AUDIO_ICON(SoundEffectSubmixPreset, SubmixEffectPreset);
 		SET_AUDIO_ICON(SoundEffectSourcePreset, SourceEffectPreset);
 		SET_AUDIO_ICON(SoundEffectSourcePresetChain, SourceEffectPresetChain_1);
diff --git a/Engine/Source/Editor/AudioEditor/Private/SoundModulationDestinationLayout.cpp b/Engine/Source/Editor/AudioEditor/Private/SoundModulationDestinationLayout.cpp
--- a/Engine/Source/Editor/AudioEditor/Private/SoundModulationDestinationLayout.cpp
+++ b/Engine/Source/Editor/AudioEditor/Private/SoundModulationDestinationLayout.cpp
@@ -436,4 +436,35 @@ void FSoundModulationDestinationLayoutCustomization::CustomizeChildren(TSharedRe
 		ModDestinationLayoutUtils::CustomizeChildren_AddValueNoModRow(ChildBuilder, StructPropertyHandle, ValueHandle);
 	}
 }
+
+void FSoundModulationDefaultSettingsLayoutCustomization::CustomizeHeader(TSharedRef<IPropertyHandle> StructPropertyHandle, FDetailWidgetRow& HeaderRow, IPropertyTypeCustomizationUtils& CustomizationUtils)
+{
+}
+
+void FSoundModulationDefaultSettingsLayoutCustomization::CustomizeChildren(TSharedRef<IPropertyHandle> StructPropertyHandle, IDetailChildrenBuilder& ChildBuilder, IPropertyTypeCustomizationUtils& StructCustomizationUtils)
+{
+	if (ModDestinationLayoutUtils::IsModulationEnabled())
+	{
+		TMap<FName, TSharedPtr<IPropertyHandle>> PropertyHandles;
+		uint32 NumChildren;
+		StructPropertyHandle->GetNumChildren(NumChildren);
+
+		for (uint32 ChildIndex = 0; ChildIndex < NumChildren; ++ChildIndex)
+		{
+			TSharedRef<IPropertyHandle> ChildHandle = StructPropertyHandle->GetChildHandle(ChildIndex).ToSharedRef();
+			const FName PropertyName = ChildHandle->GetProperty()->GetFName();
+			PropertyHandles.Add(PropertyName, ChildHandle);
+		}
+
+		TSharedRef<IPropertyHandle> VolumeHandle = PropertyHandles.FindChecked(GET_MEMBER_NAME_CHECKED(FSoundModulationDefaultSettings, VolumeModulationDestination)).ToSharedRef();
+		TSharedRef<IPropertyHandle> PitchHandle = PropertyHandles.FindChecked(GET_MEMBER_NAME_CHECKED(FSoundModulationDefaultSettings, PitchModulationDestination)).ToSharedRef();
+		TSharedRef<IPropertyHandle> HighpassHandle = PropertyHandles.FindChecked(GET_MEMBER_NAME_CHECKED(FSoundModulationDefaultSettings, HighpassModulationDestination)).ToSharedRef();
+		TSharedRef<IPropertyHandle> LowpassHandle = PropertyHandles.FindChecked(GET_MEMBER_NAME_CHECKED(FSoundModulationDefaultSettings, LowpassModulationDestination)).ToSharedRef();
+
+		ChildBuilder.AddProperty(VolumeHandle);
+		ChildBuilder.AddProperty(PitchHandle);
+		ChildBuilder.AddProperty(HighpassHandle);
+		ChildBuilder.AddProperty(LowpassHandle);
+	}
+}
 #undef LOCTEXT_NAMESPACE
diff --git a/Engine/Source/Editor/AudioEditor/Private/SoundModulationDestinationLayout.h b/Engine/Source/Editor/AudioEditor/Private/SoundModulationDestinationLayout.h
--- a/Engine/Source/Editor/AudioEditor/Private/SoundModulationDestinationLayout.h
+++ b/Engine/Source/Editor/AudioEditor/Private/SoundModulationDestinationLayout.h
@@ -14,7 +14,6 @@
 // Forward Declarations
 struct FSoundModulationDestinationSettings;
 
-
 class FSoundModulationDestinationLayoutCustomization : public IPropertyTypeCustomization
 {
 public:
@@ -28,3 +27,17 @@ public:
 	virtual void CustomizeChildren(TSharedRef<IPropertyHandle> StructPropertyHandle, IDetailChildrenBuilder& ChildBuilder, IPropertyTypeCustomizationUtils& StructCustomizationUtils) override;
 	//~ End IPropertyTypeCustomization
 };
+
+class FSoundModulationDefaultSettingsLayoutCustomization : public IPropertyTypeCustomization
+{
+public:
+	static TSharedRef<IPropertyTypeCustomization> MakeInstance()
+	{
+		return MakeShared<FSoundModulationDefaultSettingsLayoutCustomization>();
+	}
+
+	//~ Begin IPropertyTypeCustomization
+	virtual void CustomizeHeader(TSharedRef<IPropertyHandle> StructPropertyHandle, FDetailWidgetRow& HeaderRow, IPropertyTypeCustomizationUtils& StructCustomizationUtils) override;
+	virtual void CustomizeChildren(TSharedRef<IPropertyHandle> StructPropertyHandle, IDetailChildrenBuilder& ChildBuilder, IPropertyTypeCustomizationUtils& StructCustomizationUtils) override;
+	//~ End IPropertyTypeCustomization
+};
\ No newline at end of file
diff --git a/Engine/Source/Runtime/AudioMixer/Private/AudioMixerSource.cpp b/Engine/Source/Runtime/AudioMixer/Private/AudioMixerSource.cpp
--- a/Engine/Source/Runtime/AudioMixer/Private/AudioMixerSource.cpp
+++ b/Engine/Source/Runtime/AudioMixer/Private/AudioMixerSource.cpp
@@ -124,17 +124,11 @@ namespace Audio
 
 			if (USoundBase* Sound = WaveInstance->ActiveSound->GetSound())
 			{
-				InitParams.VolumeModulationSettings		= Sound->VolumeModulationDestination;
-				InitParams.PitchModulationSettings		= Sound->PitchModulationDestination;
-				InitParams.HighpassModulationSettings	= Sound->HighpassModulationDestination;
-				InitParams.LowpassModulationSettings	= Sound->LowpassModulationDestination;
+				InitParams.ModulationSettings = Sound->ModulationSettings;
 			}
 			else
 			{
-				InitParams.VolumeModulationSettings		= WaveInstance->WaveData->VolumeModulationDestination;
-				InitParams.PitchModulationSettings		= WaveInstance->WaveData->PitchModulationDestination;
-				InitParams.HighpassModulationSettings	= WaveInstance->WaveData->HighpassModulationDestination;
-				InitParams.LowpassModulationSettings	= WaveInstance->WaveData->LowpassModulationDestination;
+				InitParams.ModulationSettings = WaveInstance->WaveData->ModulationSettings;
 			}
 
 			if (WaveInstance->bIsAmbisonics && (WaveInstance->WaveData->NumChannels != 4))
@@ -879,6 +873,9 @@ namespace Audio
 
 		check(WaveInstance);
 
+		FActiveSound* ActiveSound = WaveInstance->ActiveSound;
+		check(ActiveSound);
+
 		Pitch = WaveInstance->GetPitch();
 
 		// Don't apply global pitch scale to UI sounds
@@ -899,10 +896,19 @@ namespace Audio
 			MixerSourceVoice->SetPitch(Pitch);
 		}
 
-		const float ModPitch = WaveInstance->ActiveSound->GetSound()
-			? WaveInstance->ActiveSound->GetSound()->PitchModulationDestination.Value
-			: WaveInstance->WaveData->PitchModulationDestination.Value;
-		MixerSourceVoice->SetModPitch(ModPitch);
+		const FSoundModulationDestinationSettings* PitchSettings = nullptr;
+		if (USoundBase* Sound = ActiveSound->GetSound())
+		{
+			PitchSettings = &Sound->ModulationSettings.PitchModulationDestination;
+		}
+		else
+		{
+			USoundWave* WaveData = WaveInstance->WaveData;
+			check(WaveData);
+			PitchSettings = &WaveData->ModulationSettings.PitchModulationDestination;
+		}
+
+		MixerSourceVoice->SetModPitch(PitchSettings->Value);
 	}
 
 	void FMixerSource::UpdateVolume()
@@ -923,10 +929,22 @@ namespace Audio
 			// 3. Apply editor gain stage(s)
 			CurrentVolume = FMath::Clamp<float>(GetDebugVolume(CurrentVolume), 0.0f, MAX_VOLUME);
 
-			const float ModVolume = WaveInstance->ActiveSound->GetSound()
-				? WaveInstance->ActiveSound->GetSound()->VolumeModulationDestination.Value
-				: WaveInstance->WaveData->VolumeModulationDestination.Value;
-			MixerSourceVoice->SetModVolume(ModVolume);
+			FActiveSound* ActiveSound = WaveInstance->ActiveSound;
+			check(ActiveSound);
+
+			const FSoundModulationDestinationSettings* VolumeSettings = nullptr;
+			if (USoundBase* Sound = ActiveSound->GetSound())
+			{
+				VolumeSettings = &Sound->ModulationSettings.VolumeModulationDestination;
+			}
+			else
+			{
+				USoundWave* WaveData = WaveInstance->WaveData;
+				check(WaveData);
+				VolumeSettings = &WaveData->ModulationSettings.VolumeModulationDestination;
+			}
+
+			MixerSourceVoice->SetModVolume(VolumeSettings->Value);
 		}
 		MixerSourceVoice->SetVolume(CurrentVolume);
 	}
@@ -959,15 +977,20 @@ namespace Audio
 			LastHPFFrequency = HPFFrequency;
 		}
 
-		USoundBase* Sound = WaveInstance->ActiveSound->GetSound();
+		check(WaveInstance);
+		FActiveSound* ActiveSound = WaveInstance->ActiveSound;
+		check(ActiveSound);
+
+		USoundBase* Sound = ActiveSound->GetSound();
 		if (!Sound)
 		{
 			Sound = WaveInstance->WaveData;
 		}
 		if (Sound)
 		{
-			MixerSourceVoice->SetModHPFFrequency(Sound->HighpassModulationDestination.Value);
-			MixerSourceVoice->SetModLPFFrequency(Sound->LowpassModulationDestination.Value);
+			const FSoundModulationDefaultSettings& Settings = Sound->ModulationSettings;
+			MixerSourceVoice->SetModHPFFrequency(Settings.HighpassModulationDestination.Value);
+			MixerSourceVoice->SetModLPFFrequency(Settings.LowpassModulationDestination.Value);
 		}
 
 		// If reverb is applied, figure out how of the source to "send" to the reverb.
diff --git a/Engine/Source/Runtime/AudioMixer/Private/AudioMixerSourceManager.cpp b/Engine/Source/Runtime/AudioMixer/Private/AudioMixerSourceManager.cpp
--- a/Engine/Source/Runtime/AudioMixer/Private/AudioMixerSourceManager.cpp
+++ b/Engine/Source/Runtime/AudioMixer/Private/AudioMixerSourceManager.cpp
@@ -707,19 +707,19 @@ namespace Audio
 
 		FModulationDestination VolumeModulation;
 		VolumeModulation.Init(MixerDevice->DeviceID, FName("Volume"), false /* bInIsBuffered */, true /* bInValueLinear */);
-		VolumeModulation.UpdateModulator(InitParams.VolumeModulationSettings.Modulator);
+		VolumeModulation.UpdateModulator(InitParams.ModulationSettings.VolumeModulationDestination.Modulator);
 
 		FModulationDestination PitchModulation;
 		PitchModulation.Init(MixerDevice->DeviceID, FName("Pitch"), false /* bInIsBuffered */);
-		PitchModulation.UpdateModulator(InitParams.PitchModulationSettings.Modulator);
+		PitchModulation.UpdateModulator(InitParams.ModulationSettings.PitchModulationDestination.Modulator);
 
 		FModulationDestination HighpassModulation;
 		HighpassModulation.Init(MixerDevice->DeviceID, FName("HPFCutoffFrequency"), false /* bInIsBuffered */);
-		HighpassModulation.UpdateModulator(InitParams.HighpassModulationSettings.Modulator);
+		HighpassModulation.UpdateModulator(InitParams.ModulationSettings.HighpassModulationDestination.Modulator);
 
 		FModulationDestination LowpassModulation;
 		LowpassModulation.Init(MixerDevice->DeviceID, FName("LPFCutoffFrequency"), false /* bInIsBuffered */);
-		LowpassModulation.UpdateModulator(InitParams.LowpassModulationSettings.Modulator);
+		LowpassModulation.UpdateModulator(InitParams.ModulationSettings.LowpassModulationDestination.Modulator);
 
 		AudioMixerThreadCommand([this, SourceId, InitParams, VolumeModulation, HighpassModulation, LowpassModulation, PitchModulation]()
 		{
diff --git a/Engine/Source/Runtime/AudioMixer/Private/AudioMixerSourceManager.h b/Engine/Source/Runtime/AudioMixer/Private/AudioMixerSourceManager.h
--- a/Engine/Source/Runtime/AudioMixer/Private/AudioMixerSourceManager.h
+++ b/Engine/Source/Runtime/AudioMixer/Private/AudioMixerSourceManager.h
@@ -120,10 +120,7 @@ namespace Audio
 		UOcclusionPluginSourceSettingsBase* OcclusionPluginSettings = nullptr;
 		UReverbPluginSourceSettingsBase* ReverbPluginSettings = nullptr;
 
-		FSoundModulationDestinationSettings VolumeModulationSettings;
-		FSoundModulationDestinationSettings PitchModulationSettings;
-		FSoundModulationDestinationSettings LowpassModulationSettings;
-		FSoundModulationDestinationSettings HighpassModulationSettings;
+		FSoundModulationDefaultSettings ModulationSettings;
 
 		FName AudioComponentUserID;
 		uint64 AudioComponentID = 0;
diff --git a/Engine/Source/Runtime/Engine/Classes/Sound/SoundBase.h b/Engine/Source/Runtime/Engine/Classes/Sound/SoundBase.h
--- a/Engine/Source/Runtime/Engine/Classes/Sound/SoundBase.h
+++ b/Engine/Source/Runtime/Engine/Classes/Sound/SoundBase.h
@@ -84,9 +84,6 @@ public:
 	uint8 bHasConcatenatorNode : 1;
 
 #if WITH_EDITORONLY_DATA
-	UPROPERTY(Transient)
-	uint8 bModulationEnabled : 1;
-
 	UPROPERTY()
 	uint8 bHasVirtualizeWhenSilent_DEPRECATED : 1;
 #endif // WITH_EDITORONLY_DATA
@@ -150,21 +147,9 @@ public:
 	UPROPERTY(EditAnywhere, Category = Attenuation)
 	USoundAttenuation* AttenuationSettings;
 
-	/** Volume modulation */
-	UPROPERTY(EditAnywhere, BlueprintReadWrite, Category = "Modulation", meta = (DisplayName = "Volume", AudioParam = "Volume"))
-	FSoundModulationDestinationSettings VolumeModulationDestination;
-
-	/** Pitch modulation */
-	UPROPERTY(EditAnywhere, BlueprintReadWrite, Category = "Modulation", meta = (DisplayName = "Pitch", AudioParam = "Pitch"))
-	FSoundModulationDestinationSettings PitchModulationDestination;
-
-	/** Highpass modulation */
-	UPROPERTY(EditAnywhere, BlueprintReadWrite, Category = "Modulation", meta = (DisplayName = "Highpass", AudioParam = "HPFCutoffFrequency"))
-	FSoundModulationDestinationSettings HighpassModulationDestination;
-
-	/** Lowpass modulation */
-	UPROPERTY(EditAnywhere, BlueprintReadWrite, Category = "Modulation", meta = (DisplayName = "Lowpass", AudioParam = "LPFCutoffFrequency"))
-	FSoundModulationDestinationSettings LowpassModulationDestination;
+	/** Modulation Settings */
+	UPROPERTY(EditAnywhere, BlueprintReadWrite, Category = "Modulation")
+	FSoundModulationDefaultSettings ModulationSettings;
 
 	/** Submix to route sound output to. If unset, falls back to referenced SoundClass submix.
 	  * If SoundClass submix is unset, sends to the 'Master Submix' as set in the 'Audio' category of Project Settings'. */
diff --git a/Engine/Source/Runtime/Engine/Classes/Sound/SoundModulationDestination.h b/Engine/Source/Runtime/Engine/Classes/Sound/SoundModulationDestination.h
--- a/Engine/Source/Runtime/Engine/Classes/Sound/SoundModulationDestination.h
+++ b/Engine/Source/Runtime/Engine/Classes/Sound/SoundModulationDestination.h
@@ -36,6 +36,31 @@ struct ENGINE_API FSoundModulationDestinationSettings
 };
 
 
+/** Default parameter destination settings for source audio object. */
+USTRUCT(BlueprintType)
+struct ENGINE_API FSoundModulationDefaultSettings
+{
+	GENERATED_USTRUCT_BODY()
+
+	FSoundModulationDefaultSettings();
+
+	/** Volume modulation */
+	UPROPERTY(EditAnywhere, BlueprintReadWrite, Category = "Modulation", meta = (DisplayName = "Volume", AudioParam = "Volume"))
+	FSoundModulationDestinationSettings VolumeModulationDestination;
+
+	/** Pitch modulation */
+	UPROPERTY(EditAnywhere, BlueprintReadWrite, Category = "Modulation", meta = (DisplayName = "Pitch", AudioParam = "Pitch"))
+	FSoundModulationDestinationSettings PitchModulationDestination;
+
+	/** Highpass modulation */
+	UPROPERTY(EditAnywhere, BlueprintReadWrite, Category = "Modulation", meta = (DisplayName = "Highpass", AudioParam = "HPFCutoffFrequency"))
+	FSoundModulationDestinationSettings HighpassModulationDestination;
+
+	/** Lowpass modulation */
+	UPROPERTY(EditAnywhere, BlueprintReadWrite, Category = "Modulation", meta = (DisplayName = "Lowpass", AudioParam = "LPFCutoffFrequency"))
+	FSoundModulationDestinationSettings LowpassModulationDestination;
+};
+
 namespace Audio
 {
 	struct ENGINE_API FModulationDestination
diff --git a/Engine/Source/Runtime/Engine/Private/SoundBase.cpp b/Engine/Source/Runtime/Engine/Private/SoundBase.cpp
--- a/Engine/Source/Runtime/Engine/Private/SoundBase.cpp
+++ b/Engine/Source/Runtime/Engine/Private/SoundBase.cpp
@@ -19,14 +19,8 @@ USoundBase::USoundBase(const FObjectInitializer& ObjectInitializer)
 	, Priority(1.0f)
 {
 #if WITH_EDITORONLY_DATA
-	bModulationEnabled = 0;
 	MaxConcurrentPlayCount_DEPRECATED = 16;
 #endif // WITH_EDITORONLY_DATA
-
-	VolumeModulationDestination.Value = 0.0f;
-	PitchModulationDestination.Value = 0.0f;
-	HighpassModulationDestination.Value = MIN_FILTER_FREQUENCY;
-	LowpassModulationDestination.Value = MAX_FILTER_FREQUENCY;
 }
 
 void USoundBase::PostInitProperties()
@@ -208,15 +202,6 @@ void USoundBase::PostLoad()
 		ConcurrencyOverrides.MaxCount = FMath::Max(MaxConcurrentPlayCount_DEPRECATED, 1);
 		ConcurrencyOverrides.ResolutionRule = MaxConcurrentResolutionRule_DEPRECATED;
 	}
-
-	if (FAudioDeviceManager* DeviceManager = FAudioDeviceManager::Get())
-	{
-		FAudioDeviceHandle DeviceHandle = DeviceManager->GetActiveAudioDevice();
-		if (FAudioDevice* AudioDevice = DeviceHandle.GetAudioDevice())
-		{
-			bModulationEnabled = AudioDevice->IsModulationPluginEnabled() && AudioDevice->ModulationInterface.IsValid();
-		}
-	}
 }
 #endif // WITH_EDITORONLY_DATA
 
diff --git a/Engine/Source/Runtime/Engine/Private/SoundModulationDestination.cpp b/Engine/Source/Runtime/Engine/Private/SoundModulationDestination.cpp
--- a/Engine/Source/Runtime/Engine/Private/SoundModulationDestination.cpp
+++ b/Engine/Source/Runtime/Engine/Private/SoundModulationDestination.cpp
@@ -8,6 +8,14 @@
 #include "UObject/Object.h"
 
 
+FSoundModulationDefaultSettings::FSoundModulationDefaultSettings()
+{
+	VolumeModulationDestination.Value = 0.0f;
+	PitchModulationDestination.Value = 0.0f;
+	HighpassModulationDestination.Value = MIN_FILTER_FREQUENCY;
+	LowpassModulationDestination.Value = MAX_FILTER_FREQUENCY;
+}
+
 namespace Audio
 {
 	FModulationDestination::FModulationDestination(const FModulationDestination& InModulationDestination)
