Edigrate from 4 to 5:
- Clean-up Generator property UX
- Add invert flag to EnvelopeFollower to allow AudioBus to gate source

diff --git a/Engine/Plugins/Runtime/AudioModulation/Source/AudioModulation/Private/Generators/SoundModulationEnvelopeFollower.cpp b/Engine/Plugins/Runtime/AudioModulation/Source/AudioModulation/Private/Generators/SoundModulationEnvelopeFollower.cpp
--- a/Engine/Plugins/Runtime/AudioModulation/Source/AudioModulation/Private/Generators/SoundModulationEnvelopeFollower.cpp
+++ b/Engine/Plugins/Runtime/AudioModulation/Source/AudioModulation/Private/Generators/SoundModulationEnvelopeFollower.cpp
@@ -35,6 +35,14 @@ namespace AudioModulation
 					{
 						AudioBusPatch = MixerDevice->AddPatchForAudioBus(BusId, Params.Gain);
 						EnvelopeFollower = FEnvelopeFollower(MixerDevice->SampleRate, Params.AttackTime, Params.ReleaseTime, Audio::EPeakMode::Peak);
+						if (Params.bInvert)
+						{
+							CurrentValue = 1.0f - EnvelopeFollower.GetCurrentValue();
+						}
+						else
+						{
+							CurrentValue = EnvelopeFollower.GetCurrentValue();
+						}
 					}
 				}
 			});
@@ -43,7 +51,7 @@ namespace AudioModulation
 
 	float FEnvelopeFollowerGenerator::GetValue() const
 	{
-		return EnvelopeFollower.GetCurrentValue();
+		return CurrentValue;
 	}
 
 	bool FEnvelopeFollowerGenerator::IsBypassed() const
@@ -85,6 +93,14 @@ namespace AudioModulation
 		}
 
 		EnvelopeFollower.ProcessAudio(Amp);
+		if (Params.bInvert)
+		{
+			CurrentValue = 1.0f - EnvelopeFollower.GetCurrentValue();
+		}
+		else
+		{
+			CurrentValue = EnvelopeFollower.GetCurrentValue();
+		}
 	}
 
 #if !UE_BUILD_SHIPPING
diff --git a/Engine/Plugins/Runtime/AudioModulation/Source/AudioModulation/Public/Generators/SoundModulationEnvelopeFollower.h b/Engine/Plugins/Runtime/AudioModulation/Source/AudioModulation/Public/Generators/SoundModulationEnvelopeFollower.h
--- a/Engine/Plugins/Runtime/AudioModulation/Source/AudioModulation/Public/Generators/SoundModulationEnvelopeFollower.h
+++ b/Engine/Plugins/Runtime/AudioModulation/Source/AudioModulation/Public/Generators/SoundModulationEnvelopeFollower.h
@@ -25,45 +25,49 @@ struct FEnvelopeFollowerGeneratorParams
 {
 	GENERATED_USTRUCT_BODY()
 
-	/** AudioBus to follow amplitude of and generate modulation control signal from. */
+	/** If true, bypasses generator from being modulated by parameters, patches, or mixed (remains active and computed). */
 	UPROPERTY(EditAnywhere, Category = Modulation, BlueprintReadWrite)
+	bool bBypass = false;
+
+	/** If true, inverts output */
+	UPROPERTY(EditAnywhere, Category = Modulation, BlueprintReadWrite, meta = (EditCondition = "!bBypass", DisplayAfter = "ReleaseTime"))
+	bool bInvert = false;
+
+	/** AudioBus to follow amplitude of and generate modulation control signal from. */
+	UPROPERTY(EditAnywhere, Category = Modulation, BlueprintReadWrite, meta = (EditCondition = "!bBypass"))
 	UAudioBus* AudioBus = nullptr;
 
 	/** Gain to apply to amplitude signal. */
-	UPROPERTY(EditAnywhere, Category = Modulation, BlueprintReadWrite)
+	UPROPERTY(EditAnywhere, Category = Modulation, BlueprintReadWrite, meta = (EditCondition = "!bBypass", ClampMin = 0.0f))
 	float Gain = 1.0f;
 
 	/** Attack time of envelope response (in sec) */
-	UPROPERTY(EditAnywhere, Category = Modulation, BlueprintReadWrite, meta = (ClampMin = 0.0f))
+	UPROPERTY(EditAnywhere, Category = Modulation, BlueprintReadWrite, meta = (EditCondition = "!bBypass", ClampMin = 0.0f))
 	float AttackTime = 0.010f;
 
 	/** Release time of envelope response (in sec) */
-	UPROPERTY(EditAnywhere, Category = Modulation, BlueprintReadWrite, meta = (ClampMin = 0.0f))
+	UPROPERTY(EditAnywhere, Category = Modulation, BlueprintReadWrite, meta = (EditCondition = "!bBypass", ClampMin = 0.0f))
 	float ReleaseTime = 0.100f;
-
-	/** If true, bypasses generator from being modulated by parameters, patches, or mixed (remains active and computed). */
-	UPROPERTY(EditAnywhere, Category = Modulation, BlueprintReadWrite)
-	bool bBypass = false;
 };
 
 namespace AudioModulation
 {
 	class AUDIOMODULATION_API FEnvelopeFollowerGenerator : public IGenerator
 	{
-		public:
-			FEnvelopeFollowerGenerator(const FEnvelopeFollowerGeneratorParams& InParams, Audio::FDeviceId InDeviceId);
-			virtual ~FEnvelopeFollowerGenerator() = default;
+	public:
+		FEnvelopeFollowerGenerator(const FEnvelopeFollowerGeneratorParams& InParams, Audio::FDeviceId InDeviceId);
+		virtual ~FEnvelopeFollowerGenerator() = default;
 
-			virtual float GetValue() const override;
-			virtual bool IsBypassed() const override;
-			virtual void Update(double InElapsed) override;
+		virtual float GetValue() const override;
+		virtual bool IsBypassed() const override;
+		virtual void Update(double InElapsed) override;
 
 #if !UE_BUILD_SHIPPING
-			static const FString DebugName;
+		static const FString DebugName;
 
-			virtual void GetDebugCategories(TArray<FString>& OutDebugCategories) const override;
-			virtual void GetDebugValues(TArray<FString>& OutDebugValues) const override;
-			virtual const FString& GetDebugName() const override;
+		virtual void GetDebugCategories(TArray<FString>& OutDebugCategories) const override;
+		virtual void GetDebugValues(TArray<FString>& OutDebugValues) const override;
+		virtual const FString& GetDebugName() const override;
 #endif // !UE_BUILD_SHIPPING
 
 	protected:
@@ -73,6 +77,8 @@ namespace AudioModulation
 		Audio::FPatchOutputStrongPtr AudioBusPatch;
 		Audio::FAlignedFloatBuffer TempBuffer;
 		Audio::FEnvelopeFollower EnvelopeFollower;
+
+		float CurrentValue = 0.0f;
 	};
 } // namespace AudioModulation
 
diff --git a/Engine/Plugins/Runtime/AudioModulation/Source/AudioModulation/Public/Generators/SoundModulationLFO.h b/Engine/Plugins/Runtime/AudioModulation/Source/AudioModulation/Public/Generators/SoundModulationLFO.h
--- a/Engine/Plugins/Runtime/AudioModulation/Source/AudioModulation/Public/Generators/SoundModulationLFO.h
+++ b/Engine/Plugins/Runtime/AudioModulation/Source/AudioModulation/Public/Generators/SoundModulationLFO.h
@@ -32,27 +32,27 @@ struct FSoundModulationLFOParams
 	GENERATED_USTRUCT_BODY()
 
 	/** Shape of oscillating waveform */
-	UPROPERTY(EditAnywhere, Category = Modulation, BlueprintReadWrite)
+	UPROPERTY(EditAnywhere, Category = Modulation, BlueprintReadWrite, meta = (DisplayPriority = 20, EditCondition = "!bBypass"))
 	FSoundModulationLFOShape Shape;
 
 	/** Amplitude of oscillator */
-	UPROPERTY(EditAnywhere, Category = Modulation, BlueprintReadWrite, meta = (UIMin = "0", UIMax = "1", ClampMin = "0", ClampMax = "1"))
+	UPROPERTY(EditAnywhere, Category = Modulation, BlueprintReadWrite, meta = (DisplayPriority = 30, EditCondition = "!bBypass", UIMin = "0", UIMax = "1", ClampMin = "0", ClampMax = "1"))
 	float Amplitude = 0.5f;
 
 	/** Frequency of oscillator */
-	UPROPERTY(EditAnywhere, Category = Modulation, BlueprintReadWrite, meta = (UIMin = "0", UIMax = "20", ClampMin = "0", ClampMax = "20"))
+	UPROPERTY(EditAnywhere, Category = Modulation, BlueprintReadWrite, meta = (DisplayPriority = 40, EditCondition = "!bBypass", UIMin = "0", UIMax = "20", ClampMin = "0", ClampMax = "20"))
 	float Frequency = 1.0f;
 
 	/** Offset of oscillator */
-	UPROPERTY(EditAnywhere, Category = Modulation, BlueprintReadWrite, meta = (UIMin = "0", UIMax = "1", ClampMin = "0", ClampMax = "1"))
+	UPROPERTY(EditAnywhere, Category = Modulation, BlueprintReadWrite, meta = (DisplayPriority = 50, EditCondition = "!bBypass", UIMin = "0", UIMax = "1", ClampMin = "0", ClampMax = "1"))
 	float Offset = 0.5;
 
 	/** Whether or not to loop the oscillation more than once */
-	UPROPERTY(EditAnywhere, Category = Modulation, BlueprintReadWrite)
+	UPROPERTY(EditAnywhere, Category = Modulation, BlueprintReadWrite, meta = (DisplayPriority = 60, EditCondition = "!bBypass"))
 	bool bLooping = 1;
 
 	/** If true, bypasses LFO bus from being modulated by parameters, patches, or mixed (LFO remains active and computed). */
-	UPROPERTY(EditAnywhere, Category = Modulation, BlueprintReadWrite)
+	UPROPERTY(EditAnywhere, Category = Modulation, BlueprintReadWrite, meta = (DisplayPriority = 10))
 	bool bBypass = 0;
 };
 
