Fix pedantic spelling issue

diff --git a/Engine/Plugins/Runtime/AudioModulation/Source/AudioModulation/Private/AudioModulationSystem.cpp b/Engine/Plugins/Runtime/AudioModulation/Source/AudioModulation/Private/AudioModulationSystem.cpp
--- a/Engine/Plugins/Runtime/AudioModulation/Source/AudioModulation/Private/AudioModulationSystem.cpp
+++ b/Engine/Plugins/Runtime/AudioModulation/Source/AudioModulation/Private/AudioModulationSystem.cpp
@@ -638,12 +638,12 @@ namespace AudioModulation
 	{
 		if (bInUpdateObject)
 		{
-			TMap<uint32, const FSoundControlBusMixStage*> UpdatedStageBusses;
+			TMap<uint32, const FSoundControlBusMixStage*> UpdatedStageBuses;
 			for (const FSoundControlBusMixStage& Stage : InStages)
 			{
 				if (Stage.Bus)
 				{
-					UpdatedStageBusses.Add(Stage.Bus->GetUniqueID(), &Stage);
+					UpdatedStageBuses.Add(Stage.Bus->GetUniqueID(), &Stage);
 				}
 			}
 
@@ -655,7 +655,7 @@ namespace AudioModulation
 					continue;
 				}
 
-				if (const FSoundControlBusMixStage* BusStage = UpdatedStageBusses.FindRef(Stage.Bus->GetUniqueID()))
+				if (const FSoundControlBusMixStage* BusStage = UpdatedStageBuses.FindRef(Stage.Bus->GetUniqueID()))
 				{
 					Stage = *BusStage;
 					bMarkDirty = true;
