Add missing AUDIOMODULATION_API decorator

diff --git a/Engine/Plugins/Runtime/AudioModulation/Source/AudioModulation/Public/AudioModulationStatics.h b/Engine/Plugins/Runtime/AudioModulation/Source/AudioModulation/Public/AudioModulationStatics.h
--- a/Engine/Plugins/Runtime/AudioModulation/Source/AudioModulation/Public/AudioModulationStatics.h
+++ b/Engine/Plugins/Runtime/AudioModulation/Source/AudioModulation/Public/AudioModulationStatics.h
@@ -20,7 +20,7 @@ namespace AudioModulation
 
 
 UCLASS()
-class UAudioModulationStatics : public UBlueprintFunctionLibrary
+class AUDIOMODULATION_API UAudioModulationStatics : public UBlueprintFunctionLibrary
 {
 	GENERATED_UCLASS_BODY()
 
