Replace thread check with conditional thread ID update in AudioModulation to avoid crash when suspending game when going to home screen (which can swap the AudioMixer's ThreadID on resume)

diff --git a/Engine/Plugins/Runtime/AudioModulation/Source/AudioModulation/Private/AudioModulationSystem.cpp b/Engine/Plugins/Runtime/AudioModulation/Source/AudioModulation/Private/AudioModulationSystem.cpp
--- a/Engine/Plugins/Runtime/AudioModulation/Source/AudioModulation/Private/AudioModulationSystem.cpp
+++ b/Engine/Plugins/Runtime/AudioModulation/Source/AudioModulation/Private/AudioModulationSystem.cpp
@@ -335,7 +335,8 @@ namespace AudioModulation
 
 	void FAudioModulationSystem::ProcessModulators(const double InElapsed)
 	{
-		check(ProcessingThreadId == 0 || IsInProcessingThread());
+		// The processing thread can get updated between frames. As modulation
+		// processing should be first call in frame, update the threadId accordingly
 		ProcessingThreadId = FPlatformTLS::GetCurrentThreadId();
 
 		int32 CommandsProcessed = 0;
diff --git a/Engine/Plugins/Runtime/AudioModulation/Source/AudioModulation/Private/AudioModulationSystem.h b/Engine/Plugins/Runtime/AudioModulation/Source/AudioModulation/Private/AudioModulationSystem.h
--- a/Engine/Plugins/Runtime/AudioModulation/Source/AudioModulation/Private/AudioModulationSystem.h
+++ b/Engine/Plugins/Runtime/AudioModulation/Source/AudioModulation/Private/AudioModulationSystem.h
@@ -17,6 +17,7 @@
 #include "SoundModulationValue.h"
 #include "SoundModulationGenerator.h"
 #include "SoundModulationGeneratorProxy.h"
+#include "Templates/Atomic.h"
 #include "Templates/Function.h"
 
 #if WITH_AUDIOMODULATION
@@ -177,7 +178,7 @@ namespace AudioModulation
 		TQueue<TUniqueFunction<void()>, EQueueMode::Mpsc> ProcessingThreadCommandQueue;
 
 		// Thread modulators are processed on
-		uint32 ProcessingThreadId = 0;
+		TAtomic<uint32> ProcessingThreadId { 0 };
 
 		// Collection of maps with modulator handles to referencing object ids used by externally managing objects
 		FReferencedModulators RefModulators;
