Fix for crash in AudioModulation when user leaves parameter blank, which results in no mix function being set in a ParameterPatch

diff --git a/Engine/Plugins/Runtime/AudioModulation/Source/AudioModulation/Private/SoundModulationPatchProxy.h b/Engine/Plugins/Runtime/AudioModulation/Source/AudioModulation/Private/SoundModulationPatchProxy.h
--- a/Engine/Plugins/Runtime/AudioModulation/Source/AudioModulation/Private/SoundModulationPatchProxy.h
+++ b/Engine/Plugins/Runtime/AudioModulation/Source/AudioModulation/Private/SoundModulationPatchProxy.h
@@ -73,7 +73,10 @@ namespace AudioModulation
 		/** Function used to mix patch inputs together */
 		Audio::FModulationMixFunction MixFunction;
 
-		FModulationPatchSettings() = default;
+		FModulationPatchSettings()
+			: MixFunction(Audio::FModulationParameter::GetDefaultMixFunction())
+		{
+		}
 
 		FModulationPatchSettings(const FSoundControlModulationPatch& InPatch, Audio::FDeviceId InDeviceId)
 			: bBypass(InPatch.bBypass)
@@ -83,6 +86,10 @@ namespace AudioModulation
 				DefaultOutputValue = InPatch.OutputParameter->Settings.ValueNormalized;
 				MixFunction = InPatch.OutputParameter->GetMixFunction();
 			}
+			else
+			{
+				MixFunction = Audio::FModulationParameter::GetDefaultMixFunction();
+			}
 
 			for (const FSoundControlModulationInput& Input : InPatch.Inputs)
 			{
@@ -101,6 +108,10 @@ namespace AudioModulation
 			{
 				MixFunction = Parameter->GetMixFunction();
 			}
+			else
+			{
+				MixFunction = Audio::FModulationParameter::GetDefaultMixFunction();
+			}
 
 			for (const FSoundControlModulationInput& Input : InPatch.PatchSettings.Inputs)
 			{
