Fix misnamed enum that is not initialized

diff --git a/Engine/Plugins/Runtime/AudioModulation/Config/BaseAudioModulation.ini b/Engine/Plugins/Runtime/AudioModulation/Config/BaseAudioModulation.ini
--- a/Engine/Plugins/Runtime/AudioModulation/Config/BaseAudioModulation.ini
+++ b/Engine/Plugins/Runtime/AudioModulation/Config/BaseAudioModulation.ini
@@ -1,3 +1,6 @@
+[CoreRedirects]
++EnumRedirects=(OldName="FSoundModulationLFOShape",NewName="ESoundModulationLFOShape")
+
 [/Script/AudioModulation.AudioModulationSettings]
 +Parameters=/AudioModulation/BitDepth.BitDepth
 +Parameters=/AudioModulation/HPFCutoffFrequency.HPFCutoffFrequency
diff --git a/Engine/Plugins/Runtime/AudioModulation/Source/AudioModulation/Public/Generators/SoundModulationLFO.h b/Engine/Plugins/Runtime/AudioModulation/Source/AudioModulation/Public/Generators/SoundModulationLFO.h
--- a/Engine/Plugins/Runtime/AudioModulation/Source/AudioModulation/Public/Generators/SoundModulationLFO.h
+++ b/Engine/Plugins/Runtime/AudioModulation/Source/AudioModulation/Public/Generators/SoundModulationLFO.h
@@ -13,7 +13,7 @@
 
 
 UENUM(BlueprintType)
-enum class FSoundModulationLFOShape : uint8
+enum class ESoundModulationLFOShape : uint8
 {
 	Sine			  UMETA(DisplayName = "Sine"),
 	UpSaw			  UMETA(DisplayName = "Saw (Up)"),
@@ -33,7 +33,7 @@ struct FSoundModulationLFOParams
 
 	/** Shape of oscillating waveform */
 	UPROPERTY(EditAnywhere, Category = Modulation, BlueprintReadWrite, meta = (DisplayPriority = 20, EditCondition = "!bBypass"))
-	FSoundModulationLFOShape Shape;
+	ESoundModulationLFOShape Shape = ESoundModulationLFOShape::Sine;
 
 	/** Amplitude of oscillator */
 	UPROPERTY(EditAnywhere, Category = Modulation, BlueprintReadWrite, meta = (DisplayPriority = 30, EditCondition = "!bBypass", UIMin = "0", UIMax = "1", ClampMin = "0", ClampMax = "1"))
@@ -69,7 +69,7 @@ namespace AudioModulation
 				LFO.SetFrequency(Params.Frequency);
 				LFO.SetMode(Params.bLooping ? Audio::ELFOMode::Type::Sync : Audio::ELFOMode::OneShot);
 
-				static_assert(static_cast<int32>(FSoundModulationLFOShape::COUNT) == static_cast<int32>(Audio::ELFO::Type::NumLFOTypes), "LFOShape/ELFO Type mismatch");
+				static_assert(static_cast<int32>(ESoundModulationLFOShape::COUNT) == static_cast<int32>(Audio::ELFO::Type::NumLFOTypes), "LFOShape/ELFO Type mismatch");
 				LFO.SetType(static_cast<Audio::ELFO::Type>(Params.Shape));
 				LFO.Start();
 			}
