Modulation Feature: Global Bus Value Mixes
- Enables a way to quickly assignment of values, debug easier, and mix lazily at a top-level for persistent buses.

diff --git a/Engine/Plugins/Runtime/AudioModulation/Source/AudioModulation/Private/AudioModulation.cpp b/Engine/Plugins/Runtime/AudioModulation/Source/AudioModulation/Private/AudioModulation.cpp
--- a/Engine/Plugins/Runtime/AudioModulation/Source/AudioModulation/Private/AudioModulation.cpp
+++ b/Engine/Plugins/Runtime/AudioModulation/Source/AudioModulation/Private/AudioModulation.cpp
@@ -140,6 +140,21 @@ namespace AudioModulation
 		ModSystem->SoloBusMix(InBusMix);
 	}
 
+	void FAudioModulation::SetGlobalBusMixValue(USoundControlBus& InBus, float InValue, float InFadeTime)
+	{
+		ModSystem->SetGlobalBusMixValue(InBus, InValue, InFadeTime);
+	}
+
+	void FAudioModulation::ClearGlobalBusMixValue(const USoundControlBus& InBus, float InFadeTime)
+	{
+		ModSystem->ClearGlobalBusMixValue(InBus, InFadeTime);
+	}
+
+	void FAudioModulation::ClearAllGlobalBusMixValues(float InFadeTime)
+	{
+		ModSystem->ClearAllGlobalBusMixValues(InFadeTime);
+	}
+
 #if !UE_BUILD_SHIPPING
 	bool FAudioModulation::OnPostHelp(FCommonViewportClient* ViewportClient, const TCHAR* Stream)
 	{
diff --git a/Engine/Plugins/Runtime/AudioModulation/Source/AudioModulation/Private/AudioModulationStatics.cpp b/Engine/Plugins/Runtime/AudioModulation/Source/AudioModulation/Private/AudioModulationStatics.cpp
--- a/Engine/Plugins/Runtime/AudioModulation/Source/AudioModulation/Private/AudioModulationStatics.cpp
+++ b/Engine/Plugins/Runtime/AudioModulation/Source/AudioModulation/Private/AudioModulationStatics.cpp
@@ -334,6 +334,39 @@ void UAudioModulationStatics::UpdateMix(const UObject* WorldContextObject, USoun
 	}
 }
 
+void UAudioModulationStatics::SetGlobalBusMixValue(const UObject* WorldContextObject, USoundControlBus* Bus, float Value, float FadeTime)
+{
+	if (Bus)
+	{
+		UWorld* World = GetAudioWorld(WorldContextObject);
+		if (AudioModulation::FAudioModulation* ModSystem = GetModulation(World))
+		{
+			ModSystem->SetGlobalBusMixValue(*Bus, Value, FadeTime);
+		}
+	}
+}
+
+void UAudioModulationStatics::ClearGlobalBusMixValue(const UObject* WorldContextObject, USoundControlBus* Bus, float FadeTime)
+{
+	if (Bus)
+	{
+		UWorld* World = GetAudioWorld(WorldContextObject);
+		if (AudioModulation::FAudioModulation* ModSystem = GetModulation(World))
+		{
+			ModSystem->ClearGlobalBusMixValue(*Bus, FadeTime);
+		}
+	}
+}
+
+void UAudioModulationStatics::ClearAllGlobalBusMixValues(const UObject* WorldContextObject, float FadeTime)
+{
+	UWorld* World = GetAudioWorld(WorldContextObject);
+	if (AudioModulation::FAudioModulation* ModSystem = GetModulation(World))
+	{
+		ModSystem->ClearAllGlobalBusMixValues(FadeTime);
+	}
+}
+
 void UAudioModulationStatics::UpdateMixByFilter(
 	const UObject* WorldContextObject,
 	USoundControlBusMix* Mix,
diff --git a/Engine/Plugins/Runtime/AudioModulation/Source/AudioModulation/Private/AudioModulationSystem.cpp b/Engine/Plugins/Runtime/AudioModulation/Source/AudioModulation/Private/AudioModulationSystem.cpp
--- a/Engine/Plugins/Runtime/AudioModulation/Source/AudioModulation/Private/AudioModulationSystem.cpp
+++ b/Engine/Plugins/Runtime/AudioModulation/Source/AudioModulation/Private/AudioModulationSystem.cpp
@@ -149,6 +149,8 @@ namespace AudioModulation
 
 	void FAudioModulationSystem::DeactivateBus(const USoundControlBus& InBus)
 	{
+		ClearGlobalBusMixValue(InBus);
+
 		RunCommandOnProcessingThread([this, BusId = static_cast<FBusId>(InBus.GetUniqueID())]()
 		{
 			FBusHandle BusHandle = FBusHandle::Get(BusId, RefProxies.Buses);
@@ -174,6 +176,8 @@ namespace AudioModulation
 
 	void FAudioModulationSystem::DeactivateAllBusMixes()
 	{
+		ClearAllGlobalBusMixValues();
+
 		RunCommandOnProcessingThread([this]()
 		{
 			for (TPair<FBusMixId, FModulatorBusMixProxy>& Pair : RefProxies.BusMixes)
@@ -227,6 +231,82 @@ namespace AudioModulation
 	}
 #endif // !UE_BUILD_SHIPPING
 
+	void FAudioModulationSystem::SetGlobalBusMixValue(USoundControlBus& InBus, float InValue, float InFadeTime)
+	{
+		if (TObjectPtr<USoundControlBusMix> GlobalMix = ActiveGlobalBusValueMixes.FindRef(InBus.GetUniqueID()))
+		{
+			if (ensure(!GlobalMix->MixStages.IsEmpty()))
+			{
+				GlobalMix->MixStages[0].Value.TargetValue = InValue;
+				UpdateMix(*GlobalMix, InFadeTime);
+
+				UE_LOG(LogAudioModulation, VeryVerbose, TEXT("GlobalBusMix for ControlBus '%s' updated, target set to '%0.4f'."), *InBus.GetName(), InValue);
+			}
+		}
+		else
+		{
+			const FString MixName = InBus.GetName() + TEXT("_GlobalMix");
+			TObjectPtr<USoundControlBusMix> NewGlobalMix = NewObject<USoundControlBusMix>(GetTransientPackage(), FName(*MixName));
+
+			{
+				FSoundModulationMixValue MixValue;
+				MixValue.TargetValue = InValue;
+
+				if (InFadeTime >= 0.0f)
+				{
+					MixValue.AttackTime = InFadeTime;
+				}
+
+				FSoundControlBusMixStage MixStage;
+				MixStage.Bus = &InBus;
+				MixStage.Value = MixValue;
+
+				NewGlobalMix->MixStages.Emplace(MoveTemp(MixStage));
+			}
+
+			ActiveGlobalBusValueMixes.Add(InBus.GetUniqueID(), NewGlobalMix);
+			UE_LOG(LogAudioModulation, VeryVerbose, TEXT("GlobalBusMix for ControlBus '%s' activated, target set to '%0.4f'."), *InBus.GetName(), InValue);
+			ActivateBusMix(*NewGlobalMix);
+			NewGlobalMix->AddToRoot();
+		}
+
+	}
+
+	void FAudioModulationSystem::ClearGlobalBusMixValue(const USoundControlBus& InBus, float InFadeTime)
+	{
+		const uint32 BusID = InBus.GetUniqueID();
+		if (TObjectPtr<USoundControlBusMix> GlobalMix = ActiveGlobalBusValueMixes.FindRef(BusID))
+		{
+			if (ensure(!GlobalMix->MixStages.IsEmpty()))
+			{
+				GlobalMix->MixStages[0].Value.ReleaseTime = InFadeTime;
+				DeactivateBusMix(*GlobalMix);
+				ActiveGlobalBusValueMixes.Remove(BusID);
+				GlobalMix->RemoveFromRoot();
+				UE_LOG(LogAudioModulation, VeryVerbose, TEXT("GlobalBusMix for ControlBus '%s' cleared."), *InBus.GetName());
+			}
+		}
+		else
+		{
+			UE_LOG(LogAudioModulation, VeryVerbose, TEXT("GlobalBusMix for ControlBus '%s' not active, ignoring clear request."), *InBus.GetName());
+		}
+	}
+
+	void FAudioModulationSystem::ClearAllGlobalBusMixValues(float InFadeTime)
+	{
+		TArray<TObjectPtr<USoundControlBusMix>> GlobalBusMixes;
+		ActiveGlobalBusValueMixes.GenerateValueArray(GlobalBusMixes);
+		for (const TObjectPtr<USoundControlBusMix>& BusMix : GlobalBusMixes)
+		{
+			if (ensure(!BusMix->MixStages.IsEmpty()))
+			{
+				ClearGlobalBusMixValue(*BusMix->MixStages[0].Bus, InFadeTime);
+			}
+		}
+
+		ActiveGlobalBusValueMixes.Reset();
+	}
+
 	bool FAudioModulationSystem::GetModulatorValue(const Audio::FModulatorHandle& InModulatorHandle, float& OutValue) const
 	{
 		const EModulatorType ModulatorType = static_cast<EModulatorType>(InModulatorHandle.GetTypeId());
@@ -475,7 +555,7 @@ namespace AudioModulation
 				}
 						
 				TMap<FBusId, FSoundModulationMixValue> StageInfo = PassedStageInfo;
-				USoundControlBusMix* TempMix = NewObject<USoundControlBusMix>(GetTransientPackage(), *FGuid().ToString(EGuidFormats::Short));
+				USoundControlBusMix* TempMix = NewObject<USoundControlBusMix>(GetTransientPackage(), *FGuid::NewGuid().ToString(EGuidFormats::Short));
 
 				// Buses on proxy may differ than those on uobject definition, so iterate and find by cached ids
 				// and add to temp mix to be serialized.
diff --git a/Engine/Plugins/Runtime/AudioModulation/Source/AudioModulation/Private/AudioModulationSystem.h b/Engine/Plugins/Runtime/AudioModulation/Source/AudioModulation/Private/AudioModulationSystem.h
--- a/Engine/Plugins/Runtime/AudioModulation/Source/AudioModulation/Private/AudioModulationSystem.h
+++ b/Engine/Plugins/Runtime/AudioModulation/Source/AudioModulation/Private/AudioModulationSystem.h
@@ -100,6 +100,15 @@ namespace AudioModulation
 		 */
 		void UpdateMix(const USoundControlBusMix& InMix, float InFadeTime = -1.0f);
 
+		/* Sets the global bus mix value if over the prescribed time. If FadeTime is non-positive, applies the value immediately. */
+		void SetGlobalBusMixValue(USoundControlBus& Bus, float Value, float FadeTime = -1.0f);
+
+		/* Clears the global bus mix value over the prescribed FadeTime. If FadeTime is non-positive, returns to the bus's respective parameter default immediately. */
+		void ClearGlobalBusMixValue(const USoundControlBus& InBus, float FadeTime = -1.0f);
+
+		/* Clears all global bus mix values over the prescribed FadeTime. If FadeTime is non-positive, returns to the bus's respective parameter default immediately. */
+		void ClearAllGlobalBusMixValues(float FadeTime = -1.0f);
+
 		/*
 		 * Commits any changes from a modulator type applied to a UObject definition
 		 * to modulator instance if active (i.e. Control Bus, Control Bus Modulator)
@@ -174,6 +183,9 @@ namespace AudioModulation
 		TSet<FBusMixHandle> ManuallyActivatedBusMixes;
 		TSet<FGeneratorHandle> ManuallyActivatedGenerators;
 
+		// Global mixes each containing a single stage for any globally manipulated bus stage value.
+		TMap<uint32, TObjectPtr<USoundControlBusMix>> ActiveGlobalBusValueMixes;
+
 		// Command queue to be consumed on processing thread 
 		TQueue<TUniqueFunction<void()>, EQueueMode::Mpsc> ProcessingThreadCommandQueue;
 
@@ -251,6 +263,10 @@ namespace AudioModulation
 		void SaveMixToProfile(const USoundControlBusMix& InBusMix, const int32 InProfileIndex) { }
 		void LoadMixFromProfile(const int32 InProfileIndex, USoundControlBusMix& OutBusMix) { }
 
+		void SetGlobalBusMixValue(USoundControlBus& Bus, float Value, float FadeTime) { }
+		void ClearGlobalBusMixValue(const USoundControlBus& InBus, float FadeTime) { }
+		void ClearAllGlobalBusMixValues(float FadeTime) { }
+
 		void ProcessModulators(const double InElapsed) { }
 
 		Audio::FModulatorTypeId RegisterModulator(Audio::FModulatorHandleId InHandleId, const USoundModulatorBase* InModulatorBase, Audio::FModulationParameter& OutParameter) { return INDEX_NONE; }
diff --git a/Engine/Plugins/Runtime/AudioModulation/Source/AudioModulation/Private/SoundModulationValue.cpp b/Engine/Plugins/Runtime/AudioModulation/Source/AudioModulation/Private/SoundModulationValue.cpp
--- a/Engine/Plugins/Runtime/AudioModulation/Source/AudioModulation/Private/SoundModulationValue.cpp
+++ b/Engine/Plugins/Runtime/AudioModulation/Source/AudioModulation/Private/SoundModulationValue.cpp
@@ -54,7 +54,14 @@ void FSoundModulationMixValue::SetActiveFade(EActiveFade InActiveFade, float InF
 			// If fade was not set prior, use attack time as default.
 			else if (LerpTime < 0.0f)
 			{
-				LerpTime = AttackTime;
+				if (ensureAlwaysMsgf(AttackTime >= 0.0f, TEXT("Lerp time must be non-negative if attack time not set and overriding.")))
+				{
+					LerpTime = AttackTime;
+				}
+				else
+				{
+					LerpTime = 0.0f;
+				}
 			}
 			break;
 		}
@@ -106,8 +113,14 @@ void FSoundModulationMixValue::UpdateDelta()
 	// Initialize to attack time if unset
 	if (LerpTime < 0.0f)
 	{
-		check(ActiveFade == EActiveFade::Attack);
-		LerpTime = AttackTime;
+		if (ensureAlwaysMsgf(ActiveFade == EActiveFade::Attack, TEXT("Lerp time must be non-negative if not attacking.")))
+		{
+			LerpTime = AttackTime;
+		}
+		else
+		{
+			LerpTime = 0.0f;
+		}
 	}
 
 	if (LerpTime > 0.0f)
diff --git a/Engine/Plugins/Runtime/AudioModulation/Source/AudioModulation/Public/AudioModulation.h b/Engine/Plugins/Runtime/AudioModulation/Source/AudioModulation/Public/AudioModulation.h
--- a/Engine/Plugins/Runtime/AudioModulation/Source/AudioModulation/Public/AudioModulation.h
+++ b/Engine/Plugins/Runtime/AudioModulation/Source/AudioModulation/Public/AudioModulation.h
@@ -55,6 +55,10 @@ namespace AudioModulation
 
 		void SoloBusMix(const USoundControlBusMix& InBusMix);
 
+		void SetGlobalBusMixValue(USoundControlBus& InBus, float InValue, float InFadeTime);
+		void ClearGlobalBusMixValue(const USoundControlBus& InBus, float InFadeTime);
+		void ClearAllGlobalBusMixValues(float InFadeTime);
+
 #if !UE_BUILD_SHIPPING
 		virtual void SetDebugBusFilter(const FString* InFilter);
 		virtual void SetDebugMixFilter(const FString* InFilter);
diff --git a/Engine/Plugins/Runtime/AudioModulation/Source/AudioModulation/Public/AudioModulationStatics.h b/Engine/Plugins/Runtime/AudioModulation/Source/AudioModulation/Public/AudioModulationStatics.h
--- a/Engine/Plugins/Runtime/AudioModulation/Source/AudioModulation/Public/AudioModulationStatics.h
+++ b/Engine/Plugins/Runtime/AudioModulation/Source/AudioModulation/Public/AudioModulationStatics.h
@@ -92,19 +92,19 @@ public:
 		float AttackTime = 0.1f,
 		float ReleaseTime = 0.1f);
 
-	/** Creates a modulation bus mix, with a bus channel set to the provided target value
+	/** Creates a modulation bus mix, with a bus stage set to the provided target value
 	 * @param Name - Name of mix.
 	 * @param Stages - Stages mix is responsible for.
 	 * @param Activate - Whether or not to activate mix on creation.
 	 */
 	UFUNCTION(BlueprintCallable, Category = "Audio", DisplayName = "Create Control Bus Mix", meta = (
-		WorldContext = "WorldContextObject", 
+		WorldContext = "WorldContextObject",
 		Keywords = "make create control bus mix modulation modulator")
 	)
 	static USoundControlBusMix* CreateBusMix(
 		const UObject* WorldContextObject,
 		FName Name, 
-		TArray<FSoundControlBusMixStage> Stages, 
+		TArray<FSoundControlBusMixStage> Stages,
 		bool Activate);
 
 	/** Deactivates a bus. Does nothing if the provided bus is already inactive
@@ -160,12 +160,12 @@ public:
 	)
 	static UPARAM(DisplayName = "Stages") TArray<FSoundControlBusMixStage> LoadMixFromProfile(const UObject* WorldContextObject, USoundControlBusMix* Mix, bool bActivate = true, int32 ProfileIndex = 0);
 
-	/** Sets a Control Bus Mix with the provided channel data, if the channels
+	/** Sets a Control Bus Mix with the provided stage data, if the stages
 	 *  are provided in an active instance proxy of the mix. 
 	 *  Does not update UObject definition of the mix. 
 	 * @param Mix - Mix to update
 	 * @param Stages - Stages to set.  If stage's bus is not referenced by mix, stage's update request is ignored.
-	 * @param InFadeTime - Fade time to user when interpolating between current value and new values.
+	 * @param FadeTime - Fade time to user when interpolating between current value and new values.
 	 *					 If negative, falls back to last fade time set on stage. If fade time never set on stage,
 	 *					 uses attack time set on stage in mix asset.
 	 */
@@ -175,6 +175,41 @@ public:
 	)
 	static void UpdateMix(const UObject* WorldContextObject, USoundControlBusMix* Mix, TArray<FSoundControlBusMixStage> Stages, float FadeTime = -1.0f);
 
+	/** Sets a Global Control Bus Mix with a single stage associated with the provided Bus to the given float value.  This call should
+	 * be reserved for buses that are to be always active. It is *NOT* recommended for transient buses, as not calling clear can keep
+	 * buses active indefinitely.
+	 * @param Bus - Bus associated with mix to update
+	 * @param Value - Value to set global stage to.
+	 * @param FadeTime - Fade time to user when interpolating between current value and new value. If negative, falls back to last fade
+	 * time set on stage. If fade time never set on stage, defaults to 100ms.
+	 */
+	UFUNCTION(BlueprintCallable, Category = "Audio", DisplayName = "Set Global Control Bus Mix Value", meta = (
+		WorldContext = "WorldContextObject",
+		Keywords = "set bus control modulation modulator mix stage global")
+	)
+	static void SetGlobalBusMixValue(const UObject* WorldContextObject, USoundControlBus* Bus, float Value, float FadeTime = -1.0f);
+
+	/** Clears global control bus mix if set, using the applied fade time to return to the provided bus's parameter default value.
+	 * @param Bus - Bus associated with mix to update
+	 * @param FadeTime - Fade time to user when interpolating between current value and new values.
+	 *					 If non-positive, change is immediate.
+	 */
+	UFUNCTION(BlueprintCallable, Category = "Audio", DisplayName = "Clear Global Control Bus Mix Value", meta = (
+		WorldContext = "WorldContextObject",
+		Keywords = "clear bus control modulation modulator mix stage global")
+	)
+	static void ClearGlobalBusMixValue(const UObject* WorldContextObject, USoundControlBus* Bus, float FadeTime = -1.0f);
+
+	/** Clears all global control bus mix values if set, using the applied fade time to return all to their respective bus's parameter default value.
+	 * @param FadeTime - Fade time to user when interpolating between current value and new values.
+	 *					 If non-positive, change is immediate.
+	 */
+	UFUNCTION(BlueprintCallable, Category = "Audio", DisplayName = "Clear All Global Control Bus Mix Value", meta = (
+		WorldContext = "WorldContextObject",
+		Keywords = "clear bus control modulation modulator mix stage global")
+	)
+	static void ClearAllGlobalBusMixValues(const UObject* WorldContextObject, float FadeTime = -1.0f);
+
 	/** Sets filtered stages of a given class to a provided target value for active instance of mix.
 	 * Does not update UObject definition of mix.
 	 * @param Mix - Mix to modify
