Fixed missed default output value initialization on Sound Modulation Patch Proxies, affecting patches with output parameters that needed defaults other than 1.f, such as pitch.

diff --git a/Engine/Plugins/Runtime/AudioModulation/Source/AudioModulation/Private/SoundModulationPatchProxy.h b/Engine/Plugins/Runtime/AudioModulation/Source/AudioModulation/Private/SoundModulationPatchProxy.h
--- a/Engine/Plugins/Runtime/AudioModulation/Source/AudioModulation/Private/SoundModulationPatchProxy.h
+++ b/Engine/Plugins/Runtime/AudioModulation/Source/AudioModulation/Private/SoundModulationPatchProxy.h
@@ -106,6 +106,7 @@ namespace AudioModulation
 		{
 			if (USoundModulationParameter* Parameter = InPatch.PatchSettings.OutputParameter)
 			{
+				DefaultOutputValue = Parameter->Settings.ValueNormalized;
 				MixFunction = Parameter->GetMixFunction();
 			}
 			else
