Split Modulation Registration into actual registration/Modulation proxy generation vs Settings construction to support AudioDevice-less DATAProxy generation
- Made modulator settings objects stateless & removed AudioDevice dependency in mod settings constructors.
- Required moving ModulationParameters to thread-safe registered singleton registry in AudioExtensions and out of plugin
- Rename IAudioModulation/FAudioModulation to IAudioModulationManager/FAudioModulationManager for clarity
#rnx

diff --git a/Engine/Plugins/Runtime/AudioModulation/Source/AudioModulation/Private/AudioModulation.cpp b/Engine/Plugins/Runtime/AudioModulation/Source/AudioModulation/Private/AudioModulation.cpp
--- a/Engine/Plugins/Runtime/AudioModulation/Source/AudioModulation/Private/AudioModulation.cpp
+++ b/Engine/Plugins/Runtime/AudioModulation/Source/AudioModulation/Private/AudioModulation.cpp
@@ -2,6 +2,7 @@
 #include "AudioModulation.h"
 
 #include "AudioModulationLogging.h"
+#include "AudioModulationSettings.h"
 #include "AudioModulationSystem.h"
 #include "CanvasTypes.h"
 #include "Features/IModularFeatures.h"
@@ -13,6 +14,7 @@
 #include "SoundModulationParameter.h"
 #include "SoundModulationPatch.h"
 #include "SoundModulatorAsset.h"
+#include "UObject/NoExportTypes.h"
 
 
 REGISTER_METASOUND_DATATYPE(AudioModulation::FSoundModulatorAsset, "Modulator", Metasound::ELiteralType::UObjectProxy, USoundModulatorBase);
@@ -20,206 +22,244 @@ REGISTER_METASOUND_DATATYPE(AudioModulation::FSoundModulationParameterAsset, "Mo
 
 namespace AudioModulation
 {
-	FAudioModulation::FAudioModulation()
+	FAudioModulationManager::FAudioModulationManager()
 		: ModSystem(new FAudioModulationSystem())
 	{
 	}
 
-	FAudioModulation::~FAudioModulation()
+	FAudioModulationManager::~FAudioModulationManager()
 	{
 		delete ModSystem;
 	}
 
-	Audio::FModulationParameter FAudioModulation::GetParameter(FName InParamName)
-	{
-		return ModSystem->GetParameter(InParamName);
-	}
-
-	void FAudioModulation::Initialize(const FAudioPluginInitializationParams& InitializationParams)
+	void FAudioModulationManager::Initialize(const FAudioPluginInitializationParams& InitializationParams)
 	{
 		ModSystem->Initialize(InitializationParams);
 	}
 
-	void FAudioModulation::OnAuditionEnd()
+	void FAudioModulationManager::OnAuditionEnd()
 	{
 		ModSystem->OnAuditionEnd();
 	}
 
-	void FAudioModulation::ActivateBus(const USoundControlBus& InBus)
+	void FAudioModulationManager::ActivateBus(const USoundControlBus& InBus)
 	{
 		ModSystem->ActivateBus(InBus);
 	}
 
-	void FAudioModulation::ActivateBusMix(const USoundControlBusMix& InBusMix)
+	void FAudioModulationManager::ActivateBusMix(const USoundControlBusMix& InBusMix)
 	{
 		ModSystem->ActivateBusMix(InBusMix);
 	}
 
-	void FAudioModulation::ActivateGenerator(const USoundModulationGenerator& InGenerator)
+	void FAudioModulationManager::ActivateGenerator(const USoundModulationGenerator& InGenerator)
 	{
 		ModSystem->ActivateGenerator(InGenerator);
 	}
 
-	void FAudioModulation::DeactivateBus(const USoundControlBus& InBus)
+	void FAudioModulationManager::DeactivateBus(const USoundControlBus& InBus)
 	{
 		ModSystem->DeactivateBus(InBus);
 	}
 
-	void FAudioModulation::DeactivateBusMix(const USoundControlBusMix& InBusMix)
+	void FAudioModulationManager::DeactivateBusMix(const USoundControlBusMix& InBusMix)
 	{
 		ModSystem->DeactivateBusMix(InBusMix);
 	}
 
-	void FAudioModulation::DeactivateAllBusMixes()
+	void FAudioModulationManager::DeactivateAllBusMixes()
 	{
 		ModSystem->DeactivateAllBusMixes();
 	}
 
-	void FAudioModulation::DeactivateGenerator(const USoundModulationGenerator& InGenerator)
+	void FAudioModulationManager::DeactivateGenerator(const USoundModulationGenerator& InGenerator)
 	{
 		ModSystem->DeactivateGenerator(InGenerator);
 	}
 
 #if !UE_BUILD_SHIPPING
-	void FAudioModulation::SetDebugBusFilter(const FString* InNameFilter)
+	void FAudioModulationManager::SetDebugBusFilter(const FString* InNameFilter)
 	{
 		ModSystem->SetDebugBusFilter(InNameFilter);
 	}
 
-	void FAudioModulation::SetDebugGeneratorFilter(const FString* InFilter)
+	void FAudioModulationManager::SetDebugGeneratorFilter(const FString* InFilter)
 	{
 		ModSystem->SetDebugGeneratorFilter(InFilter);
 	}
 
-	void FAudioModulation::SetDebugGeneratorTypeFilter(const FString* InFilter, bool bInIsEnabled)
+	void FAudioModulationManager::SetDebugGeneratorTypeFilter(const FString* InFilter, bool bInIsEnabled)
 	{
 		ModSystem->SetDebugGeneratorTypeFilter(InFilter, bInIsEnabled);
 	}
 
-	void FAudioModulation::SetDebugGeneratorsEnabled(bool bInIsEnabled)
+	void FAudioModulationManager::SetDebugGeneratorsEnabled(bool bInIsEnabled)
 	{
 		ModSystem->SetDebugGeneratorsEnabled(bInIsEnabled);
 	}
 
-	void FAudioModulation::SetDebugMatrixEnabled(bool bInIsEnabled)
+	void FAudioModulationManager::SetDebugMatrixEnabled(bool bInIsEnabled)
 	{
 		ModSystem->SetDebugMatrixEnabled(bInIsEnabled);
 	}
 
-	void FAudioModulation::SetDebugMixFilter(const FString* InNameFilter)
+	void FAudioModulationManager::SetDebugMixFilter(const FString* InNameFilter)
 	{
 		ModSystem->SetDebugMixFilter(InNameFilter);
 	}
 
 #endif // !UE_BUILD_SHIPPING
 
-	void FAudioModulation::SaveMixToProfile(const USoundControlBusMix& InBusMix, const int32 InProfileIndex)
+	void FAudioModulationManager::SaveMixToProfile(const USoundControlBusMix& InBusMix, const int32 InProfileIndex)
 	{
 		ModSystem->SaveMixToProfile(InBusMix, InProfileIndex);
 	}
 
-	TArray<FSoundControlBusMixStage> FAudioModulation::LoadMixFromProfile(const int32 InProfileIndex, USoundControlBusMix& OutBusMix)
+	TArray<FSoundControlBusMixStage> FAudioModulationManager::LoadMixFromProfile(const int32 InProfileIndex, USoundControlBusMix& OutBusMix)
 	{
 		return ModSystem->LoadMixFromProfile(InProfileIndex, OutBusMix);
 	}
 
-	void FAudioModulation::UpdateMix(const TArray<FSoundControlBusMixStage>& InStages, USoundControlBusMix& InOutMix, bool bInUpdateObject, float InFadeTime)
+	void FAudioModulationManager::UpdateMix(const TArray<FSoundControlBusMixStage>& InStages, USoundControlBusMix& InOutMix, bool bInUpdateObject, float InFadeTime)
 	{
 		ModSystem->UpdateMix(InStages, InOutMix, bInUpdateObject, InFadeTime);
 	}
 
-	void FAudioModulation::UpdateMix(const USoundControlBusMix& InMix, float InFadeTime)
+	void FAudioModulationManager::UpdateMix(const USoundControlBusMix& InMix, float InFadeTime)
 	{
 		ModSystem->UpdateMix(InMix, InFadeTime);
 	}
 
-	void FAudioModulation::UpdateMixByFilter(const FString& InAddressFilter, const TSubclassOf<USoundModulationParameter>& InParamClassFilter, USoundModulationParameter* InParamFilter, float Value, float FadeTime, USoundControlBusMix& InOutMix, bool bInUpdateObject)
+	void FAudioModulationManager::UpdateMixByFilter(const FString& InAddressFilter, const TSubclassOf<USoundModulationParameter>& InParamClassFilter, USoundModulationParameter* InParamFilter, float Value, float FadeTime, USoundControlBusMix& InOutMix, bool bInUpdateObject)
 	{
 		ModSystem->UpdateMixByFilter(InAddressFilter, InParamClassFilter, InParamFilter, Value, FadeTime, InOutMix, bInUpdateObject);
 	}
 
-	void FAudioModulation::SoloBusMix(const USoundControlBusMix& InBusMix)
+	void FAudioModulationManager::SoloBusMix(const USoundControlBusMix& InBusMix)
 	{
 		ModSystem->SoloBusMix(InBusMix);
 	}
 
-	void FAudioModulation::SetGlobalBusMixValue(USoundControlBus& InBus, float InValue, float InFadeTime)
+	void FAudioModulationManager::SetGlobalBusMixValue(USoundControlBus& InBus, float InValue, float InFadeTime)
 	{
 		ModSystem->SetGlobalBusMixValue(InBus, InValue, InFadeTime);
 	}
 
-	void FAudioModulation::ClearGlobalBusMixValue(const USoundControlBus& InBus, float InFadeTime)
+	void FAudioModulationManager::ClearGlobalBusMixValue(const USoundControlBus& InBus, float InFadeTime)
 	{
 		ModSystem->ClearGlobalBusMixValue(InBus, InFadeTime);
 	}
 
-	void FAudioModulation::ClearAllGlobalBusMixValues(float InFadeTime)
+	void FAudioModulationManager::ClearAllGlobalBusMixValues(float InFadeTime)
 	{
 		ModSystem->ClearAllGlobalBusMixValues(InFadeTime);
 	}
 
 #if !UE_BUILD_SHIPPING
-	bool FAudioModulation::OnPostHelp(FCommonViewportClient* ViewportClient, const TCHAR* Stream)
+	bool FAudioModulationManager::OnPostHelp(FCommonViewportClient* ViewportClient, const TCHAR* Stream)
 	{
 		return ModSystem->OnPostHelp(ViewportClient, Stream);
 	}
 
-	int32 FAudioModulation::OnRenderStat(FViewport* Viewport, FCanvas* Canvas, int32 X, int32 Y, const UFont& Font, const FVector* ViewLocation, const FRotator* ViewRotation)
+	int32 FAudioModulationManager::OnRenderStat(FViewport* Viewport, FCanvas* Canvas, int32 X, int32 Y, const UFont& Font, const FVector* ViewLocation, const FRotator* ViewRotation)
 	{
 		return ModSystem->OnRenderStat(Viewport, Canvas, X, Y, Font, ViewLocation, ViewRotation);
 	}
 
-	bool FAudioModulation::OnToggleStat(FCommonViewportClient* ViewportClient, const TCHAR* Stream)
+	bool FAudioModulationManager::OnToggleStat(FCommonViewportClient* ViewportClient, const TCHAR* Stream)
 	{
 		return ModSystem->OnToggleStat(ViewportClient, Stream);
 	}
 #endif // !UE_BUILD_SHIPPING
 
-	void FAudioModulation::ProcessModulators(const double InElapsed)
+	void FAudioModulationManager::ProcessModulators(const double InElapsed)
 	{
 		ModSystem->ProcessModulators(InElapsed);
 	}
 
-	Audio::FModulatorTypeId FAudioModulation::RegisterModulator(Audio::FModulatorHandleId InHandleId, const USoundModulatorBase* InModulatorBase, Audio::FModulationParameter& OutParameter)
-	{
-		return ModSystem->RegisterModulator(InHandleId, InModulatorBase, OutParameter);
-	}
-
-	void FAudioModulation::RegisterModulator(Audio::FModulatorHandleId InHandleId, Audio::FModulatorId InModulatorId)
+	void FAudioModulationManager::RegisterModulator(Audio::FModulatorHandleId InHandleId, Audio::FModulatorId InModulatorId)
 	{
 		ModSystem->RegisterModulator(InHandleId, InModulatorId);
 	}
 
-	bool FAudioModulation::GetModulatorValue(const Audio::FModulatorHandle& ModulatorHandle, float& OutValue) const
+	bool FAudioModulationManager::GetModulatorValue(const Audio::FModulatorHandle& ModulatorHandle, float& OutValue) const
 	{
 		return ModSystem->GetModulatorValue(ModulatorHandle, OutValue);
 	}
 
-	bool FAudioModulation::GetModulatorValueThreadSafe(const Audio::FModulatorHandle& ModulatorHandle, float& OutValue) const
+	bool FAudioModulationManager::GetModulatorValueThreadSafe(const Audio::FModulatorHandle& ModulatorHandle, float& OutValue) const
 	{
 		return ModSystem->GetModulatorValueThreadSafe(ModulatorHandle, OutValue);
 	}
 
-	void FAudioModulation::UnregisterModulator(const Audio::FModulatorHandle& InHandle)
+	FAudioModulationSystem& FAudioModulationManager::GetSystem()
+	{
+		check(ModSystem);
+		return *ModSystem;
+	}
+
+	void FAudioModulationManager::UnregisterModulator(const Audio::FModulatorHandle& InHandle)
 	{
 		ModSystem->UnregisterModulator(InHandle);
 	}
 
-	void FAudioModulation::UpdateModulator(const USoundModulatorBase& InModulator)
+	void FAudioModulationManager::UpdateModulator(const USoundModulatorBase& InModulator)
 	{
 		ModSystem->UpdateModulator(InModulator);
 	}
+
+	FAudioModulationManager* GetDeviceModulationManager(Audio::FDeviceId InDeviceId)
+	{
+		FAudioDeviceManager* DeviceManager = FAudioDeviceManager::Get();
+		if (ensure(DeviceManager))
+		{
+			FAudioDevice* AudioDevice = DeviceManager->GetAudioDeviceRaw(InDeviceId);
+			if (ensure(AudioDevice))
+			{
+				if (AudioDevice->IsModulationPluginEnabled())
+				{
+					if (IAudioModulationManager* Modulation = AudioDevice->ModulationInterface.Get())
+					{
+						return static_cast<FAudioModulationManager*>(Modulation);
+					}
+				}
+			}
+		}
+
+		return nullptr;
+	}
+
+	void IterateModulationManagers(TFunctionRef<void(FAudioModulationManager&)> InFunction)
+	{
+		if (FAudioDeviceManager* DeviceManager = FAudioDeviceManager::Get())
+		{
+			TArray<FAudioDevice*> Devices = DeviceManager->GetAudioDevices();
+			DeviceManager->IterateOverAllDevices([ModFunction = MoveTemp(InFunction)](Audio::FDeviceId DeviceId, FAudioDevice* AudioDevice)
+			{
+				if (AudioDevice && AudioDevice->IsModulationPluginEnabled() && AudioDevice->ModulationInterface.IsValid())
+				{
+					auto ModulationInterface = static_cast<AudioModulation::FAudioModulationManager*>(AudioDevice->ModulationInterface.Get());
+					ModFunction(*ModulationInterface);
+				}
+			});
+		}
+	}
 } // namespace AudioModulation
 
 TAudioModulationPtr FAudioModulationPluginFactory::CreateNewModulationPlugin(FAudioDevice* OwningDevice)
 {
-	return TAudioModulationPtr(new AudioModulation::FAudioModulation());
+	return TAudioModulationPtr(new AudioModulation::FAudioModulationManager());
 }
 
 void FAudioModulationModule::StartupModule()
 {
 	IModularFeatures::Get().RegisterModularFeature(FAudioModulationPluginFactory::GetModularFeatureName(), &ModulationPluginFactory);
 
+	if (const UAudioModulationSettings* ModulationSettings = GetDefault<UAudioModulationSettings>())
+	{
+		ModulationSettings->RegisterParameters();
+	}
+
 	// flush node registration queue to guarantee AudioModulation DataTypes/Nodes are ready prior to assets loading
 	FMetasoundFrontendRegistryContainer::Get()->RegisterPendingNodes();
 
@@ -232,4 +272,5 @@ void FAudioModulationModule::ShutdownModule()
 	UE_LOG(LogAudioModulation, Log, TEXT("Audio Modulation Shutdown"));
 }
 
-IMPLEMENT_MODULE(FAudioModulationModule, AudioModulation);
\ No newline at end of file
+
+IMPLEMENT_MODULE(FAudioModulationModule, AudioModulation);
diff --git a/Engine/Plugins/Runtime/AudioModulation/Source/AudioModulation/Private/AudioModulationDebugger.cpp b/Engine/Plugins/Runtime/AudioModulation/Source/AudioModulation/Private/AudioModulationDebugger.cpp
--- a/Engine/Plugins/Runtime/AudioModulation/Source/AudioModulation/Private/AudioModulationDebugger.cpp
+++ b/Engine/Plugins/Runtime/AudioModulation/Source/AudioModulation/Private/AudioModulationDebugger.cpp
@@ -24,7 +24,7 @@ namespace AudioModulation
 {
 	namespace Debug
 	{
-		using FDebugViewUpdateFunction = TUniqueFunction<void(FAudioModulation& /*InModulation*/, const FString* /*InFilter*/)>;
+		using FDebugViewUpdateFunction = TUniqueFunction<void(FAudioModulationManager& /*InModulation*/, const FString* /*InFilter*/)>;
 		void UpdateObjectFilter(const TArray<FString>& Args, UWorld* World, FDebugViewUpdateFunction InUpdateViewFunc)
 		{
 			if (!World)
@@ -46,9 +46,9 @@ namespace AudioModulation
 					GEngine->Exec(World, TEXT("au.Debug.SoundModulators 1"));
 				}
 
-				if (IAudioModulation* Modulation = DeviceHandle->ModulationInterface.Get())
+				if (IAudioModulationManager* Modulation = DeviceHandle->ModulationInterface.Get())
 				{
-					FAudioModulation& ModulationImpl = *static_cast<FAudioModulation*>(Modulation);
+					FAudioModulationManager& ModulationImpl = *static_cast<FAudioModulationManager*>(Modulation);
 					InUpdateViewFunc(ModulationImpl, FilterString);
 				}
 			}
@@ -72,7 +72,7 @@ static FAutoConsoleCommandWithWorldAndArgs CVarAudioModulationDebugBuses(
 	FConsoleCommandWithWorldAndArgsDelegate::CreateLambda([](const TArray<FString>& Args, UWorld* World)
 	{
 		using namespace AudioModulation;
-		Debug::UpdateObjectFilter(Args, World, [](FAudioModulation& InModulation, const FString* InFilter)
+		Debug::UpdateObjectFilter(Args, World, [](FAudioModulationManager& InModulation, const FString* InFilter)
 		{
 			InModulation.SetDebugBusFilter(InFilter);
 		});
@@ -85,7 +85,7 @@ static FAutoConsoleCommandWithWorldAndArgs CVarAudioModulationDebugMixes(
 	FConsoleCommandWithWorldAndArgsDelegate::CreateLambda([](const TArray<FString>& Args, UWorld* World)
 	{
 		using namespace AudioModulation;
-		Debug::UpdateObjectFilter(Args, World, [](FAudioModulation& InModulation, const FString* InFilter)
+		Debug::UpdateObjectFilter(Args, World, [](FAudioModulationManager& InModulation, const FString* InFilter)
 		{
 			InModulation.SetDebugMixFilter(InFilter);
 		});
@@ -98,7 +98,7 @@ static FAutoConsoleCommandWithWorldAndArgs CVarAudioModulationDebugMatrix(
 	FConsoleCommandWithWorldAndArgsDelegate::CreateLambda([](const TArray<FString>& Args, UWorld* World)
 	{
 		using namespace AudioModulation;
-		Debug::UpdateObjectFilter(Args, World, [](FAudioModulation& InModulation, const FString* InFilter)
+		Debug::UpdateObjectFilter(Args, World, [](FAudioModulationManager& InModulation, const FString* InFilter)
 		{
 			if (InFilter)
 			{
@@ -116,7 +116,7 @@ static FAutoConsoleCommandWithWorldAndArgs CVarAudioModulationSetDebugGeneratorF
 	{
 		using namespace AudioModulation;
 
-		Debug::UpdateObjectFilter(Args, World, [](FAudioModulation& InModulation, const FString* InFilter)
+		Debug::UpdateObjectFilter(Args, World, [](FAudioModulationManager& InModulation, const FString* InFilter)
 		{
 			InModulation.SetDebugGeneratorFilter(InFilter);
 		});
@@ -136,7 +136,7 @@ static FAutoConsoleCommandWithWorldAndArgs CVarAudioModulationSetDebugGeneratorF
 		{
 			bEnabled = Args[1].ToBool();
 		}
-		Debug::UpdateObjectFilter(Args, World, [bEnabled](FAudioModulation& InModulation, const FString* InFilter)
+		Debug::UpdateObjectFilter(Args, World, [bEnabled](FAudioModulationManager& InModulation, const FString* InFilter)
 		{
 			InModulation.SetDebugGeneratorTypeFilter(InFilter, bEnabled);
 		});
@@ -149,7 +149,7 @@ static FAutoConsoleCommandWithWorldAndArgs CVarAudioModulationDebugEnableGenerat
 	FConsoleCommandWithWorldAndArgsDelegate::CreateLambda([](const TArray<FString>& Args, UWorld* World)
 	{
 		using namespace AudioModulation;
-		Debug::UpdateObjectFilter(Args, World, [](FAudioModulation& InModulation, const FString* InFilter)
+		Debug::UpdateObjectFilter(Args, World, [](FAudioModulationManager& InModulation, const FString* InFilter)
 		{
 			if (InFilter)
 			{
diff --git a/Engine/Plugins/Runtime/AudioModulation/Source/AudioModulation/Private/AudioModulationLogging.h b/Engine/Plugins/Runtime/AudioModulation/Source/AudioModulation/Private/AudioModulationLogging.h
--- a/Engine/Plugins/Runtime/AudioModulation/Source/AudioModulation/Private/AudioModulationLogging.h
+++ b/Engine/Plugins/Runtime/AudioModulation/Source/AudioModulation/Private/AudioModulationLogging.h
@@ -1,7 +1,7 @@
 // Copyright Epic Games, Inc. All Rights Reserved.
 #pragma once
 
-#include "CoreMinimal.h"
+#include "Logging/LogMacros.h"
 
 
 DECLARE_LOG_CATEGORY_EXTERN(LogAudioModulation, Log, All);
diff --git a/Engine/Plugins/Runtime/AudioModulation/Source/AudioModulation/Private/AudioModulationSettings.cpp b/Engine/Plugins/Runtime/AudioModulation/Source/AudioModulation/Private/AudioModulationSettings.cpp
new file mode 100644
--- /dev/null
+++ b/Engine/Plugins/Runtime/AudioModulation/Source/AudioModulation/Private/AudioModulationSettings.cpp
@@ -0,0 +1,41 @@
+// Copyright Epic Games, Inc. All Rights Reserved.
+#include "AudioModulationSettings.h"
+
+#include "AudioModulationLogging.h"
+#include "IAudioModulation.h"
+#include "SoundModulationParameter.h"
+
+#if WITH_EDITOR
+void UAudioModulationSettings::PostEditChangeChainProperty(FPropertyChangedChainEvent& PropertyChangedEvent)
+{
+	if (PropertyChangedEvent.Property)
+	{
+		FName PropertyName = PropertyChangedEvent.Property->GetFName();
+
+		if (PropertyName == GET_MEMBER_NAME_CHECKED(UAudioModulationSettings, Parameters))
+		{
+			UE_LOG(LogAudioModulation, Display, TEXT("Reloading Parameters From Modulation Developer Settings..."));
+			RegisterParameters();
+		}
+	}
+}
+#endif // WITH_EDITOR
+
+void UAudioModulationSettings::RegisterParameters() const
+{
+	Audio::UnregisterAllModulationParameters();
+
+	for (const FSoftObjectPath& Path : Parameters)
+	{
+		if (USoundModulationParameter* Param = Cast<USoundModulationParameter>(Path.TryLoad()))
+		{
+			Audio::FModulationParameter NewParam = Param->CreateParameter();
+			Audio::RegisterModulationParameter(NewParam.ParameterName, MoveTemp(NewParam));
+			UE_LOG(LogAudioModulation, Display, TEXT("Initialized Audio Modulation Parameter '%s'"), *NewParam.ParameterName.ToString());
+		}
+		else
+		{
+			UE_LOG(LogAudioModulation, Error, TEXT("Failed to load parameter at '%s': Missing asset or invalid type."), *Path.GetAssetName());
+		}
+	}
+}
\ No newline at end of file
diff --git a/Engine/Plugins/Runtime/AudioModulation/Source/AudioModulation/Private/AudioModulationStatics.cpp b/Engine/Plugins/Runtime/AudioModulation/Source/AudioModulation/Private/AudioModulationStatics.cpp
--- a/Engine/Plugins/Runtime/AudioModulation/Source/AudioModulation/Private/AudioModulationStatics.cpp
+++ b/Engine/Plugins/Runtime/AudioModulation/Source/AudioModulation/Private/AudioModulationStatics.cpp
@@ -117,7 +117,7 @@ void UAudioModulationStatics::ActivateBus(const UObject* WorldContextObject, USo
 	}
 
 	UWorld* World = GetAudioWorld(WorldContextObject);
-	if (AudioModulation::FAudioModulation* ModSystem = GetModulation(World))
+	if (AudioModulation::FAudioModulationManager* ModSystem = GetModulation(World))
 	{
 		ModSystem->ActivateBus(*Bus);
 	}
@@ -128,7 +128,7 @@ void UAudioModulationStatics::ActivateBusMix(const UObject* WorldContextObject,
 	if (BusMix)
 	{
 		UWorld* World = GetAudioWorld(WorldContextObject);
-		if (AudioModulation::FAudioModulation* ModSystem = GetModulation(World))
+		if (AudioModulation::FAudioModulationManager* ModSystem = GetModulation(World))
 		{
 			ModSystem->ActivateBusMix(*BusMix);
 		}
@@ -138,7 +138,7 @@ void UAudioModulationStatics::ActivateBusMix(const UObject* WorldContextObject,
 void UAudioModulationStatics::ActivateGenerator(const UObject* WorldContextObject, USoundModulationGenerator* Generator)
 {
 	UWorld* World = GetAudioWorld(WorldContextObject);
-	if (AudioModulation::FAudioModulation* ModSystem = GetModulation(World))
+	if (AudioModulation::FAudioModulationManager* ModSystem = GetModulation(World))
 	{
 		if (Generator)
 		{
@@ -163,7 +163,7 @@ UWorld* UAudioModulationStatics::GetAudioWorld(const UObject* WorldContextObject
 	return World;
 }
 
-AudioModulation::FAudioModulation* UAudioModulationStatics::GetModulation(UWorld* World)
+AudioModulation::FAudioModulationManager* UAudioModulationStatics::GetModulation(UWorld* World)
 {
 	FAudioDeviceHandle AudioDevice;
 	if (World)
@@ -180,9 +180,9 @@ AudioModulation::FAudioModulation* UAudioModulationStatics::GetModulation(UWorld
 
 	if (AudioDevice.IsValid() && AudioDevice->IsModulationPluginEnabled())
 	{
-		if (IAudioModulation* ModulationInterface = AudioDevice->ModulationInterface.Get())
+		if (IAudioModulationManager* ModulationInterface = AudioDevice->ModulationInterface.Get())
 		{
-			return static_cast<AudioModulation::FAudioModulation*>(ModulationInterface);
+			return static_cast<AudioModulation::FAudioModulationManager*>(ModulationInterface);
 		}
 	}
 
@@ -203,7 +203,7 @@ USoundControlBus* UAudioModulationStatics::CreateBus(UObject* WorldContextObject
 
 	if (Activate)
 	{
-		if (AudioModulation::FAudioModulation* ModSystem = GetModulation(World))
+		if (AudioModulation::FAudioModulationManager* ModSystem = GetModulation(World))
 		{
 			ModSystem->ActivateBus(*NewBus);
 		}
@@ -245,7 +245,7 @@ USoundControlBusMix* UAudioModulationStatics::CreateBusMix(UObject* WorldContext
 
 	if (Activate)
 	{
-		if (AudioModulation::FAudioModulation* ModSystem = GetModulation(World))
+		if (AudioModulation::FAudioModulationManager* ModSystem = GetModulation(World))
 		{
 			ModSystem->ActivateBusMix(*NewBusMix);
 		}
@@ -259,7 +259,7 @@ void UAudioModulationStatics::DeactivateBus(const UObject* WorldContextObject, U
 	if (Bus)
 	{
 		UWorld* World = GetAudioWorld(WorldContextObject);
-		if (AudioModulation::FAudioModulation* ModSystem = GetModulation(World))
+		if (AudioModulation::FAudioModulationManager* ModSystem = GetModulation(World))
 		{
 			ModSystem->DeactivateBus(*Bus);
 		}
@@ -271,7 +271,7 @@ void UAudioModulationStatics::DeactivateBusMix(const UObject* WorldContextObject
 	if (BusMix)
 	{
 		UWorld* World = GetAudioWorld(WorldContextObject);
-		if (AudioModulation::FAudioModulation* ModSystem = GetModulation(World))
+		if (AudioModulation::FAudioModulationManager* ModSystem = GetModulation(World))
 		{
 			ModSystem->DeactivateBusMix(*BusMix);
 		}
@@ -281,7 +281,7 @@ void UAudioModulationStatics::DeactivateBusMix(const UObject* WorldContextObject
 void UAudioModulationStatics::DeactivateGenerator(const UObject* WorldContextObject, USoundModulationGenerator* Generator)
 {
 	UWorld* World = GetAudioWorld(WorldContextObject);
-	if (AudioModulation::FAudioModulation* ModSystem = GetModulation(World))
+	if (AudioModulation::FAudioModulationManager* ModSystem = GetModulation(World))
 	{
 		if (Generator)
 		{
@@ -293,7 +293,7 @@ void UAudioModulationStatics::DeactivateGenerator(const UObject* WorldContextObj
 void UAudioModulationStatics::SaveMixToProfile(const UObject* WorldContextObject, USoundControlBusMix* BusMix, int32 ProfileIndex)
 {
 	UWorld* World = GetAudioWorld(WorldContextObject);
-	if (AudioModulation::FAudioModulation* ModSystem = GetModulation(World))
+	if (AudioModulation::FAudioModulationManager* ModSystem = GetModulation(World))
 	{
 		if (BusMix)
 		{
@@ -307,7 +307,7 @@ TArray<FSoundControlBusMixStage> UAudioModulationStatics::LoadMixFromProfile(con
 	if (BusMix)
 	{
 		UWorld* World = GetAudioWorld(WorldContextObject);
-		if (AudioModulation::FAudioModulation* ModSystem = GetModulation(World))
+		if (AudioModulation::FAudioModulationManager* ModSystem = GetModulation(World))
 		{
 			if (bActivate)
 			{
@@ -325,7 +325,7 @@ void UAudioModulationStatics::UpdateMix(const UObject* WorldContextObject, USoun
 	if (Mix)
 	{
 		UWorld* World = GetAudioWorld(WorldContextObject);
-		if (AudioModulation::FAudioModulation* ModSystem = GetModulation(World))
+		if (AudioModulation::FAudioModulationManager* ModSystem = GetModulation(World))
 		{
 			// UObject representation is not updated in this form of the call as doing so from
 			// PIE can result in an unstable state where UObject is modified but not properly dirtied.
@@ -339,7 +339,7 @@ void UAudioModulationStatics::SetGlobalBusMixValue(const UObject* WorldContextOb
 	if (Bus)
 	{
 		UWorld* World = GetAudioWorld(WorldContextObject);
-		if (AudioModulation::FAudioModulation* ModSystem = GetModulation(World))
+		if (AudioModulation::FAudioModulationManager* ModSystem = GetModulation(World))
 		{
 			ModSystem->SetGlobalBusMixValue(*Bus, Value, FadeTime);
 		}
@@ -351,7 +351,7 @@ void UAudioModulationStatics::ClearGlobalBusMixValue(const UObject* WorldContext
 	if (Bus)
 	{
 		UWorld* World = GetAudioWorld(WorldContextObject);
-		if (AudioModulation::FAudioModulation* ModSystem = GetModulation(World))
+		if (AudioModulation::FAudioModulationManager* ModSystem = GetModulation(World))
 		{
 			ModSystem->ClearGlobalBusMixValue(*Bus, FadeTime);
 		}
@@ -361,7 +361,7 @@ void UAudioModulationStatics::ClearGlobalBusMixValue(const UObject* WorldContext
 void UAudioModulationStatics::ClearAllGlobalBusMixValues(const UObject* WorldContextObject, float FadeTime)
 {
 	UWorld* World = GetAudioWorld(WorldContextObject);
-	if (AudioModulation::FAudioModulation* ModSystem = GetModulation(World))
+	if (AudioModulation::FAudioModulationManager* ModSystem = GetModulation(World))
 	{
 		ModSystem->ClearAllGlobalBusMixValues(FadeTime);
 	}
@@ -379,7 +379,7 @@ void UAudioModulationStatics::UpdateMixByFilter(
 	if (Mix)
 	{
 		UWorld* World = GetAudioWorld(WorldContextObject);
-		if (AudioModulation::FAudioModulation* ModSystem = GetModulation(World))
+		if (AudioModulation::FAudioModulationManager* ModSystem = GetModulation(World))
 		{
 			// UObject representation is not updated in this form of the call as doing so from
 			// PIE can result in an unstable state where UObject is modified but not properly dirtied.
@@ -393,7 +393,7 @@ void UAudioModulationStatics::UpdateMixFromObject(const UObject* WorldContextObj
 	if (Mix)
 	{
 		UWorld* World = GetAudioWorld(WorldContextObject);
-		if (AudioModulation::FAudioModulation* ModSystem = GetModulation(World))
+		if (AudioModulation::FAudioModulationManager* ModSystem = GetModulation(World))
 		{
 			ModSystem->UpdateMix(*Mix, InFadeTime);
 		}
@@ -405,7 +405,7 @@ void UAudioModulationStatics::UpdateModulator(const UObject* WorldContextObject,
 	if (Modulator)
 	{
 		UWorld* World = GetAudioWorld(WorldContextObject);
-		if (AudioModulation::FAudioModulation* ModSystem = GetModulation(World))
+		if (AudioModulation::FAudioModulationManager* ModSystem = GetModulation(World))
 		{
 			ModSystem->UpdateModulator(*Modulator);
 		}
diff --git a/Engine/Plugins/Runtime/AudioModulation/Source/AudioModulation/Private/AudioModulationSystem.cpp b/Engine/Plugins/Runtime/AudioModulation/Source/AudioModulation/Private/AudioModulationSystem.cpp
--- a/Engine/Plugins/Runtime/AudioModulation/Source/AudioModulation/Private/AudioModulationSystem.cpp
+++ b/Engine/Plugins/Runtime/AudioModulation/Source/AudioModulation/Private/AudioModulationSystem.cpp
@@ -48,24 +48,6 @@ namespace AudioModulation
 	void FAudioModulationSystem::Initialize(const FAudioPluginInitializationParams& InitializationParams)
 	{
 		AudioDeviceId = InitializationParams.AudioDevicePtr->DeviceID;
-
-		// Load all parameters listed in settings and keep loaded.  Must be done on initialization as
-		// soft object paths cannot be loaded on non-game threads (data from parameters is copied to a proxy
-		// in 'GetParameter' via the audio thread)
-		if (const UAudioModulationSettings* ModulationSettings = GetDefault<UAudioModulationSettings>())
-		{
-			for (const FSoftObjectPath& Path : ModulationSettings->Parameters)
-			{
-				if (UObject* ParamObj = Path.TryLoad())
-				{
-					ParamObj->AddToRoot();
-				}
-				else
-				{
-					UE_LOG(LogAudioModulation, Error, TEXT("Failed to load modulation parameter from path '%s'"), *Path.ToString());
-				}
-			}
-		}
 	}
 
 	void FAudioModulationSystem::OnAuditionEnd()
@@ -95,7 +77,7 @@ namespace AudioModulation
 
 	void FAudioModulationSystem::ActivateBus(const USoundControlBus& InBus)
 	{
-		RunCommandOnProcessingThread([this, Settings = FControlBusSettings(InBus, AudioDeviceId)]() mutable
+		RunCommandOnProcessingThread([this, Settings = FControlBusSettings(InBus)]() mutable
 		{
 			FBusHandle BusHandle = FBusHandle::Create(MoveTemp(Settings), RefProxies.Buses, *this);
 			ManuallyActivatedBuses.Add(MoveTemp(BusHandle));
@@ -122,15 +104,23 @@ namespace AudioModulation
 
 	void FAudioModulationSystem::ActivateBusMix(const USoundControlBusMix& InBusMix)
 	{
-		ActivateBusMix(FModulatorBusMixSettings(InBusMix, AudioDeviceId));
+		ActivateBusMix(FModulatorBusMixSettings(InBusMix));
 	}
 
 	void FAudioModulationSystem::ActivateGenerator(const USoundModulationGenerator& InGenerator)
 	{
-		RunCommandOnProcessingThread([this, Settings = FModulationGeneratorSettings(InGenerator, AudioDeviceId)]() mutable
+		RunCommandOnProcessingThread([this, Settings = FModulationGeneratorSettings(InGenerator)]() mutable
 		{
-			FGeneratorHandle GeneratorHandle = FGeneratorHandle::Create(MoveTemp(Settings), RefProxies.Generators, *this);
-			ManuallyActivatedGenerators.Add(MoveTemp(GeneratorHandle));
+			FGeneratorHandle GeneratorHandle = FGeneratorHandle::Get(Settings.GetId(), RefProxies.Generators);
+			if (GeneratorHandle.IsValid())
+			{
+				ManuallyActivatedGenerators.Add(MoveTemp(GeneratorHandle));
+			}
+			else
+			{
+				GeneratorHandle = FGeneratorHandle::Create(MoveTemp(Settings), RefProxies.Generators, *this);
+				GeneratorHandle.FindProxy().Init(AudioDeviceId);
+			}
 		});
 	}
 
@@ -383,45 +373,6 @@ namespace AudioModulation
 		return AudioDeviceId;
 	}
 
-	Audio::FModulationParameter FAudioModulationSystem::GetParameter(FName InParamName) const
-	{
-		Audio::FModulationParameter Parameter;
-		if (InParamName == FName())
-		{
-			return Parameter;
-		}
-
-		if (const UAudioModulationSettings* ModulationSettings = GetDefault<UAudioModulationSettings>())
-		{
-			for (const FSoftObjectPath& Path : ModulationSettings->Parameters)
-			{
-				if (const USoundModulationParameter* Param = Cast<const USoundModulationParameter>(Path.ResolveObject()))
-				{
-					if (Param->GetFName() == InParamName)
-					{
-						Parameter.ParameterName = InParamName;
-						Parameter.bRequiresConversion = Param->RequiresUnitConversion();
-						Parameter.MixFunction = Param->GetMixFunction();
-						Parameter.UnitFunction = Param->GetUnitConversionFunction();
-						Parameter.NormalizedFunction = Param->GetNormalizedConversionFunction();
-						Parameter.DefaultValue = Param->GetUnitDefault();
-						Parameter.MinValue = Param->GetUnitMin();
-						Parameter.MaxValue = Param->GetUnitMax();
-
-#if WITH_EDITORONLY_DATA
-						Parameter.UnitDisplayName = Param->Settings.UnitDisplayName;
-#endif // WITH_EDITORONLY_DATA
-
-						return Parameter;
-					}
-				}
-			}
-		}
-
-		UE_LOG(LogAudioModulation, Error, TEXT("Audio modulation parameter '%s' not found/loaded. Modulation may be disabled for destination referencing parameter."), *InParamName.ToString());
-		return Parameter;
-	}
-
 	bool FAudioModulationSystem::IsInProcessingThread() const
 	{
 		return ProcessingThreadId == FPlatformTLS::GetCurrentThreadId();
@@ -636,32 +587,28 @@ namespace AudioModulation
 		}
 	}
 
-	Audio::FModulatorTypeId FAudioModulationSystem::RegisterModulator(Audio::FModulatorHandleId InHandleId, const USoundModulatorBase* InModulatorBase, Audio::FModulationParameter& OutParameter)
+	Audio::FModulatorTypeId FAudioModulationSystem::RegisterModulator(Audio::FModulatorHandleId InHandleId, const FControlBusSettings& InSettings)
 	{
-		OutParameter = GetParameter(OutParameter.ParameterName);
-
-		if (!InModulatorBase)
-		{
-			return INDEX_NONE;
-		}
-
-		if (RegisterModulator<FPatchHandle, USoundModulationPatch, FModulationPatchSettings, FPatchProxyMap>(InHandleId, InModulatorBase, RefProxies.Patches, RefModulators.PatchMap))
-		{
-			return static_cast<Audio::FModulatorTypeId>(EModulatorType::Patch);
-		}
-			
-		if (RegisterModulator<FBusHandle, USoundControlBus, FControlBusSettings, FBusProxyMap>(InHandleId, InModulatorBase, RefProxies.Buses, RefModulators.BusMap))
-		{
-			return static_cast<Audio::FModulatorTypeId>(EModulatorType::Bus);
-		}
+		FControlBusSettings CachedSettings = InSettings;
+		RegisterModulator(InHandleId, MoveTemp(CachedSettings), RefProxies.Buses, RefModulators.BusMap);
+		return static_cast<Audio::FModulatorTypeId>(EModulatorType::Bus);
+	}
 
-		if (RegisterModulator<FGeneratorHandle, USoundModulationGenerator, FModulationGeneratorSettings, FGeneratorProxyMap>(InHandleId, InModulatorBase, RefProxies.Generators, RefModulators.GeneratorMap))
+	Audio::FModulatorTypeId FAudioModulationSystem::RegisterModulator(Audio::FModulatorHandleId InHandleId, const FModulationGeneratorSettings& InSettings)
+	{
+		FModulationGeneratorSettings CachedSettings = InSettings;
+		RegisterModulator(InHandleId, MoveTemp(CachedSettings), RefProxies.Generators, RefModulators.GeneratorMap, [this](FGeneratorHandle& NewHandle)
 		{
-			return static_cast<Audio::FModulatorTypeId>(EModulatorType::Generator);
-		}
+			NewHandle.FindProxy().Init(AudioDeviceId);
+		});
+		return static_cast<Audio::FModulatorTypeId>(EModulatorType::Generator);
+	}
 
-		ensureMsgf(false, TEXT("Modulator type of '%s' failed to register: Unsupported type."), *InModulatorBase->GetName());
-		return INDEX_NONE;
+	Audio::FModulatorTypeId FAudioModulationSystem::RegisterModulator(Audio::FModulatorHandleId InHandleId, const FModulationPatchSettings& InSettings)
+	{
+		FModulationPatchSettings CachedSettings = InSettings;
+		RegisterModulator(InHandleId, MoveTemp(CachedSettings), RefProxies.Patches, RefModulators.PatchMap);
+		return static_cast<Audio::FModulatorTypeId>(EModulatorType::Patch);
 	}
 
 	void FAudioModulationSystem::RegisterModulator(Audio::FModulatorHandleId InHandleId, Audio::FModulatorId InModulatorId)
@@ -705,7 +652,7 @@ namespace AudioModulation
 
 	void FAudioModulationSystem::SoloBusMix(const USoundControlBusMix& InBusMix)
 	{
-		RunCommandOnProcessingThread([this, BusMixSettings = FModulatorBusMixSettings(InBusMix, AudioDeviceId)]() mutable
+		RunCommandOnProcessingThread([this, BusMixSettings = FModulatorBusMixSettings(InBusMix)]() mutable
 		{
 			bool bMixActive = false;
 			for (TPair<FBusMixId, FModulatorBusMixProxy>& Pair : RefProxies.BusMixes)
@@ -790,7 +737,7 @@ namespace AudioModulation
 		{
 			if (Stage.Bus)
 			{
-				StageSettings.Emplace(Stage, AudioDeviceId);
+				StageSettings.Emplace(Stage);
 			}
 		}
 	
@@ -870,7 +817,7 @@ namespace AudioModulation
 
 	void FAudioModulationSystem::UpdateMix(const USoundControlBusMix& InMix, float InFadeTime)
 	{
-		RunCommandOnProcessingThread([this, MixSettings = FModulatorBusMixSettings(InMix, AudioDeviceId), InFadeTime]() mutable
+		RunCommandOnProcessingThread([this, MixSettings = FModulatorBusMixSettings(InMix), InFadeTime]() mutable
 		{
 			FBusMixHandle BusMixHandle = FBusMixHandle::Get(MixSettings.GetId(), RefProxies.BusMixes);
 			if (BusMixHandle.IsValid())
@@ -898,7 +845,7 @@ namespace AudioModulation
 	{
 		if (const USoundModulationGenerator* InGenerator = Cast<USoundModulationGenerator>(&InModulator))
 		{
-			RunCommandOnProcessingThread([this, GeneratorSettings = FModulationGeneratorSettings(*InGenerator, AudioDeviceId)]() mutable
+			RunCommandOnProcessingThread([this, GeneratorSettings = FModulationGeneratorSettings(*InGenerator)]() mutable
 			{
 				FGeneratorHandle GeneratorHandle = FGeneratorHandle::Get(GeneratorSettings.GetId(), RefProxies.Generators);
 				if (GeneratorHandle.IsValid())
@@ -916,7 +863,7 @@ namespace AudioModulation
 
 		if (const USoundControlBus* InBus = Cast<USoundControlBus>(&InModulator))
 		{
-			RunCommandOnProcessingThread([this, BusSettings = FControlBusSettings(*InBus, AudioDeviceId)]() mutable
+			RunCommandOnProcessingThread([this, BusSettings = FControlBusSettings(*InBus)]() mutable
 			{
 				FBusHandle BusHandle = FBusHandle::Get(BusSettings.GetId(), RefProxies.Buses);
 				if (BusHandle.IsValid())
@@ -935,7 +882,7 @@ namespace AudioModulation
 
 		if (const USoundControlBusMix* InMix = Cast<USoundControlBusMix>(&InModulator))
 		{
-			RunCommandOnProcessingThread([this, BusMixSettings = FModulatorBusMixSettings(*InMix, AudioDeviceId)]() mutable
+			RunCommandOnProcessingThread([this, BusMixSettings = FModulatorBusMixSettings(*InMix)]() mutable
 			{
 				FBusMixHandle BusMixHandle = FBusMixHandle::Get(BusMixSettings.GetId(), RefProxies.BusMixes);
 				if (BusMixHandle.IsValid())
@@ -954,7 +901,7 @@ namespace AudioModulation
 
 		if (const USoundModulationPatch* InPatch = Cast<USoundModulationPatch>(&InModulator))
 		{
-			RunCommandOnProcessingThread([this, PatchSettings = FModulationPatchSettings(*InPatch, AudioDeviceId)]() mutable
+			RunCommandOnProcessingThread([this, PatchSettings = FModulationPatchSettings(*InPatch)]() mutable
 			{
 				FPatchHandle PatchHandle = FPatchHandle::Get(PatchSettings.GetId(), RefProxies.Patches);
 				if (PatchHandle.IsValid())
diff --git a/Engine/Plugins/Runtime/AudioModulation/Source/AudioModulation/Private/AudioModulationSystem.h b/Engine/Plugins/Runtime/AudioModulation/Source/AudioModulation/Private/AudioModulationSystem.h
--- a/Engine/Plugins/Runtime/AudioModulation/Source/AudioModulation/Private/AudioModulationSystem.h
+++ b/Engine/Plugins/Runtime/AudioModulation/Source/AudioModulation/Private/AudioModulationSystem.h
@@ -21,15 +21,14 @@
 #include "Templates/Atomic.h"
 #include "Templates/Function.h"
 
-#if WITH_AUDIOMODULATION
-
-#if !UE_BUILD_SHIPPING
-#include "AudioModulationDebugger.h"
-#endif // !UE_BUILD_SHIPPING
 
 namespace AudioModulation
 {
 	// Forward Declarations
+	struct FControlBusSettings;
+	struct FModulationGeneratorSettings;
+	struct FModulationPatchSettings;
+
 	class FControlBusProxy;
 	class FModulationInputProxy;
 	class FModulationPatchProxy;
@@ -52,7 +51,16 @@ namespace AudioModulation
 		TMap<FBusHandle, FModulatorHandleSet> BusMap;
 		TMap<FGeneratorHandle, FModulatorHandleSet> GeneratorMap;
 	};
+} // namespace AudioModulation
 
+#if WITH_AUDIOMODULATION
+
+#if !UE_BUILD_SHIPPING
+#include "AudioModulationDebugger.h"
+#endif // !UE_BUILD_SHIPPING
+
+namespace AudioModulation
+{
 	class FAudioModulationSystem
 	{
 	public:
@@ -72,15 +80,17 @@ namespace AudioModulation
 		void DeactivateAllBusMixes();
 		void DeactivateGenerator(const USoundModulationGenerator& InGenerator);
 
-		Audio::FModulationParameter GetParameter(FName InParamName) const;
-
 		void ProcessModulators(const double InElapsed);
 		void SoloBusMix(const USoundControlBusMix& InBusMix);
 
 		Audio::FDeviceId GetAudioDeviceId() const;
 
-		/* Register new handle with given a given modulator that may or may already be active (i.e. registered). */
-		Audio::FModulatorTypeId RegisterModulator(Audio::FModulatorHandleId InHandleId, const USoundModulatorBase* InModulatorBase, Audio::FModulationParameter& OutParameter);
+		/* Register new handle with given a given modulator that may or may not already be active (i.e. registered).
+		 * If already registered, depending on modulator type, may or may not refresh proxy based on provided settings.
+		 */
+		Audio::FModulatorTypeId RegisterModulator(Audio::FModulatorHandleId InHandleId, const FControlBusSettings& InSettings);
+		Audio::FModulatorTypeId RegisterModulator(Audio::FModulatorHandleId InHandleId, const FModulationGeneratorSettings& InSettings);
+		Audio::FModulatorTypeId RegisterModulator(Audio::FModulatorHandleId InHandleId, const FModulationPatchSettings& InSettings);
 
 		/* Register new handle with given Id with a modulator that is already active (i.e. registered). Used primarily for copying modulation handles. */
 		void RegisterModulator(Audio::FModulatorHandleId InHandleId, Audio::FModulatorId InModulatorId);
@@ -139,26 +149,38 @@ namespace AudioModulation
 		/* Runs the provided command on the audio render thread (at the beginning of the ProcessModulators call) */
 		void RunCommandOnProcessingThread(TUniqueFunction<void()> Cmd);
 
-		template <typename THandleType, typename TModType, typename TModSettings, typename TMapType>
-		bool RegisterModulator(Audio::FModulatorHandleId InHandleId, const USoundModulatorBase* InModulatorBase, TMapType& OutProxyMap, TMap<THandleType, FModulatorHandleSet>& OutModMap)
+		/* Template for register calls that move the modulator settings objects onto the Audio Processing Thread & create proxies.
+		 * If the proxy is already found, adds provided HandleId as reference to given proxy. Does *not* update the proxy with the
+		 * given settings.  If update is desired on an existing proxy, UpdateModulator must be explicitly called.
+		 * @tparam THandleType Type of modulator handle used to access the proxy on the Processing Thread
+		 * @tparam TModSettings Modulation settings to move and use if proxy construction is required on the Audio Processing Thread
+		 * @tparam TMapType MapType used to cache the corresponding proxy id & proxy
+		 * @tparam TInitFunction (Optional) Function type used to call on the AudioProcessingThread immediately following proxy construction
+		 * @param InHandleId HandleId associated with the proxy to be retrieved/generated.
+		 * @param InModSettings ModulatorSettings to be used if construction of proxy is required on Audio Processing Thread
+		 * @param OutProxyMap Map to find or add proxy to if constructed
+		 * @param InInitFunction Function type used to call on the AudioProcessingThread immediately following proxy generation
+		 */
+		template <typename THandleType, typename TModSettings, typename TMapType, typename TInitFunction = TUniqueFunction<void(THandleType)>>
+		void RegisterModulator(Audio::FModulatorHandleId InHandleId, TModSettings&& InModSettings, TMapType& OutProxyMap, TMap<THandleType, FModulatorHandleSet>& OutModMap, TInitFunction InInitFunction = TInitFunction())
 		{
 			check(InHandleId != INDEX_NONE);
 
-			if (const TModType* Mod = Cast<TModType>(InModulatorBase))
+			RunCommandOnProcessingThread([
+				this,
+				InHandleId,
+				ModSettings = MoveTemp(InModSettings),
+				InitFunction = MoveTemp(InInitFunction),
+				PassedProxyMap = &OutProxyMap,
+				PassedModMap = &OutModMap
+			]() mutable
 			{
-				RunCommandOnProcessingThread([this, Modulator = TModSettings(*Mod, AudioDeviceId), InHandleId, PassedProxyMap = &OutProxyMap, PassedModMap = &OutModMap]() mutable
-				{
-					check(PassedProxyMap);
-					check(PassedModMap);
-
-					THandleType Handle = THandleType::Create(MoveTemp(Modulator), *PassedProxyMap, *this);
-					PassedModMap->FindOrAdd(Handle).Add(InHandleId);
-				});
+				check(PassedProxyMap);
+				check(PassedModMap);
 
-				return true;
-			}
-
-			return false;
+				THandleType Handle = THandleType::Create(MoveTemp(ModSettings), *PassedProxyMap, *this);
+				PassedModMap->FindOrAdd(Handle).Add(InHandleId);
+			});
 		}
 
 		template <typename THandleType>
@@ -208,6 +230,10 @@ namespace AudioModulation
 
 		Audio::FDeviceId AudioDeviceId = INDEX_NONE;
 
+		// Registered Parameters
+		using FParameterRegistry = TMap<FName, Audio::FModulationParameter>;
+		FParameterRegistry ParameterRegistry;
+
 #if !UE_BUILD_SHIPPING
 	public:
 		void SetDebugBusFilter(const FString* InFilter);
@@ -232,7 +258,7 @@ namespace AudioModulation
 	};
 } // namespace AudioModulation
 
-#else // WITH_AUDIOMODULATION
+#else // !WITH_AUDIOMODULATION
 
 namespace AudioModulation
 {
@@ -240,7 +266,6 @@ namespace AudioModulation
 	class FAudioModulationSystem
 	{
 	public:
-		Audio::FModulationParameter GetParameter(FName InParamName) { return Audio::FModulationParameter(); }
 		void Initialize(const FAudioPluginInitializationParams& InitializationParams) { }
 
 #if WITH_EDITOR
@@ -268,11 +293,13 @@ namespace AudioModulation
 
 		void DeactivateAllBusMixes() { }
 		void DeactivateBus(const USoundControlBus& InBus) { }
-		void DeactivateBusMix(const USoundControlBus& InBusMix) { }
+		void DeactivateBusMix(const USoundControlBusMix& InBusMix) { }
 		void DeactivateGenerator(const USoundModulationGenerator& InGenerator) { }
 
+		Audio::FDeviceId GetAudioDeviceId() const { return 0; }
+
 		void SaveMixToProfile(const USoundControlBusMix& InBusMix, const int32 InProfileIndex) { }
-		void LoadMixFromProfile(const int32 InProfileIndex, USoundControlBusMix& OutBusMix) { }
+		TArray<FSoundControlBusMixStage> LoadMixFromProfile(const int32 InProfileIndex, USoundControlBusMix& OutBusMix) { }
 
 		void SetGlobalBusMixValue(USoundControlBus& Bus, float Value, float FadeTime) { }
 		void ClearGlobalBusMixValue(const USoundControlBus& InBus, float FadeTime) { }
@@ -280,16 +307,28 @@ namespace AudioModulation
 
 		void ProcessModulators(const double InElapsed) { }
 
-		Audio::FModulatorTypeId RegisterModulator(Audio::FModulatorHandleId InHandleId, const USoundModulatorBase* InModulatorBase, Audio::FModulationParameter& OutParameter) { return INDEX_NONE; }
+		Audio::FModulatorTypeId RegisterModulator(Audio::FModulatorHandleId InHandleId, const FControlBusSettings& InSettings) { return 0; }
+		Audio::FModulatorTypeId RegisterModulator(Audio::FModulatorHandleId InHandleId, const FModulationGeneratorSettings& InSettings) { return 0; }
+		Audio::FModulatorTypeId RegisterModulator(Audio::FModulatorHandleId InHandleId, const FModulationPatchSettings& InSettings) { return 0; }
+
 		void RegisterModulator(Audio::FModulatorHandleId InHandleId, Audio::FModulatorId InModulatorId) { }
 		bool GetModulatorValue(const Audio::FModulatorHandle& ModulatorHandle, float& OutValue) const { return false; }
 		bool GetModulatorValueThreadSafe(const Audio::FModulatorHandle& ModulatorHandle, float& OutValue) const { return false; }
 		void UnregisterModulator(const Audio::FModulatorHandle& InHandle) { }
 
-		void UpdateMix(const USoundControlBusMix& InMix, float InFadeTime) { }
-		void UpdateMix(const TArray<FSoundControlBusMixStage>& InStages, USoundControlBusMix& InOutMix, bool bUpdateObject = false) { }
+		void UpdateMix(const USoundControlBusMix& InMix, float InFadeTime = -1.0f) { }
+		void UpdateMix(const TArray<FSoundControlBusMixStage>& InStages, USoundControlBusMix& InOutMix, bool bUpdateObject = false, float InFadeTime = -1.0f) { }
 		void UpdateMixByFilter(const FString& InAddressFilter, const TSubclassOf<USoundModulationParameter>& InParamClassFilter, USoundModulationParameter* InParamFilter, float InValue, float InFadeTime, USoundControlBusMix& InOutMix, bool bUpdateObject = false) { }
 		void UpdateModulator(const USoundModulatorBase& InModulator) { }
+
+		private:
+			FReferencedProxies RefProxies;
+
+			friend FControlBusProxy;
+			friend FModulationInputProxy;
+			friend FModulationPatchProxy;
+			friend FModulationPatchRefProxy;
+			friend FModulatorBusMixStageProxy;
 	};
 } // namespace AudioModulation
 #endif // WITH_AUDIOMODULATION
diff --git a/Engine/Plugins/Runtime/AudioModulation/Source/AudioModulation/Private/Generators/SoundModulationEnvelopeFollower.cpp b/Engine/Plugins/Runtime/AudioModulation/Source/AudioModulation/Private/Generators/SoundModulationEnvelopeFollower.cpp
--- a/Engine/Plugins/Runtime/AudioModulation/Source/AudioModulation/Private/Generators/SoundModulationEnvelopeFollower.cpp
+++ b/Engine/Plugins/Runtime/AudioModulation/Source/AudioModulation/Private/Generators/SoundModulationEnvelopeFollower.cpp
@@ -3,6 +3,7 @@
 #include "Generators/SoundModulationEnvelopeFollower.h"
 
 #include "Algo/MaxElement.h"
+#include "AudioDefines.h"
 #include "AudioDeviceManager.h"
 #include "AudioMixerDevice.h"
 #include "AudioModulation.h"
@@ -26,13 +27,25 @@ namespace AudioModulation
 	class AUDIOMODULATION_API FEnvelopeFollowerGenerator : public IGenerator
 	{
 	public:
-		FEnvelopeFollowerGenerator() = default;
-		FEnvelopeFollowerGenerator(const FEnvelopeFollowerGeneratorParams& InGeneratorParams, Audio::FDeviceId InDeviceId)
-			: IGenerator(InDeviceId)
-			, Gain(InGeneratorParams.Gain)
+		FEnvelopeFollowerGenerator(const FEnvelopeFollowerGenerator& InGenerator)
+			: AudioBusPatch(InGenerator.AudioBusPatch)
+			, TempBuffer(InGenerator.TempBuffer)
+			, EnvelopeFollower(InGenerator.EnvelopeFollower)
+			, InitParams(InGenerator.InitParams)
+			, BusId(InGenerator.BusId)
+			, CurrentValue(InGenerator.CurrentValue)
+			, Gain(InGenerator.Gain)
+			, bBypass(InGenerator.bBypass)
+			, bInvert(InGenerator.bInvert)
+			, bBusRequiresPatch(InGenerator.bBusRequiresPatch)
+		{
+		}
+
+		FEnvelopeFollowerGenerator(const FEnvelopeFollowerGeneratorParams& InGeneratorParams)
+			: Gain(InGeneratorParams.Gain)
 			, bBypass(InGeneratorParams.bBypass ? 1 : 0)
 			, bInvert(InGeneratorParams.bInvert ? 1 : 0)
-			, bInitialized(0)
+			, bBusRequiresPatch(1)
 		{
 			if (UAudioBus* AudioBus = InGeneratorParams.AudioBus)
 			{
@@ -40,21 +53,10 @@ namespace AudioModulation
 				InitParams.NumChannels = AudioBus->GetNumChannels();
 			}
 
-			if (FAudioDeviceManager* DeviceManager = FAudioDeviceManager::Get())
-			{
-				FAudioDevice* AudioDevice = DeviceManager->GetAudioDeviceRaw(InDeviceId);
-				if (Audio::FMixerDevice* MixerDevice = static_cast<Audio::FMixerDevice*>(AudioDevice))
-				{
-					InitParams.SampleRate = MixerDevice->SampleRate;
-				}
-			}
-
 			InitParams.AttackTimeMsec = InGeneratorParams.AttackTime * 1000.0f;
 			InitParams.ReleaseTimeMsec = InGeneratorParams.ReleaseTime * 1000.0f;
 			InitParams.Mode = Audio::EPeakMode::Peak;
 
-			EnvelopeFollower = Audio::FEnvelopeFollower(InitParams);
-
 			if (bInvert)
 			{
 				CurrentValue = 1.0f;
@@ -119,7 +121,7 @@ namespace AudioModulation
 				{
 					BusId = Generator->BusId;
 					Gain = Generator->Gain;
-					bInitialized = false;
+					bBusRequiresPatch = true;
 				}
 			});
 		}
@@ -129,19 +131,21 @@ namespace AudioModulation
 			return CurrentValue;
 		}
 
-		void InitBus()
+		virtual void Init(Audio::FDeviceId InDeviceId) override
 		{
-			if (BusId != INDEX_NONE)
+			if (AudioDeviceId == INDEX_NONE)
 			{
+				AudioDeviceId = InDeviceId;
 				if (FAudioDeviceManager* DeviceManager = FAudioDeviceManager::Get())
 				{
-					FAudioDevice* AudioDevice = DeviceManager->GetAudioDeviceRaw(AudioDeviceId);
+					FAudioDevice* AudioDevice = DeviceManager->GetAudioDeviceRaw(InDeviceId);
 					if (Audio::FMixerDevice* MixerDevice = static_cast<Audio::FMixerDevice*>(AudioDevice))
 					{
-						AudioBusPatch = MixerDevice->AddPatchForAudioBus(BusId, Gain);
-						bInitialized = true;
+						InitParams.SampleRate = MixerDevice->SampleRate;
 					}
 				}
+
+				EnvelopeFollower = Audio::FEnvelopeFollower(InitParams);
 			}
 		}
 
@@ -152,9 +156,9 @@ namespace AudioModulation
 
 		virtual void Update(double InElapsed) override
 		{
-			if (!bInitialized)
+			if (bBusRequiresPatch)
 			{
-				InitBus();
+				CreatePatchForBus();
 			}
 
 			if (AudioBusPatch.IsValid() && !AudioBusPatch->IsInputStale())
@@ -178,23 +182,32 @@ namespace AudioModulation
 					MaxValue = FMath::Clamp(*MaxEnvelopePtr, 0.f, 1.f);
 				}
 
-				if (bInvert)
-				{
-					CurrentValue = 1.0f - MaxValue;
-				}
-				else
-				{
-					CurrentValue = MaxValue;
-				}
+				CurrentValue = bInvert ? 1.0f - MaxValue : MaxValue;
 			}
 			else
 			{
-				bInitialized = false;
-				CurrentValue = 0.f;
+				bBusRequiresPatch = true;
+				CurrentValue = bInvert ? 1.f : 0.f;
 			}
 
 		}
 
+		void CreatePatchForBus()
+		{
+			if (BusId != INDEX_NONE)
+			{
+				if (FAudioDeviceManager* DeviceManager = FAudioDeviceManager::Get())
+				{
+					FAudioDevice* AudioDevice = DeviceManager->GetAudioDeviceRaw(AudioDeviceId);
+					if (Audio::FMixerDevice* MixerDevice = static_cast<Audio::FMixerDevice*>(AudioDevice))
+					{
+						AudioBusPatch = MixerDevice->AddPatchForAudioBus(BusId, Gain);
+						bBusRequiresPatch = false;
+					}
+				}
+			}
+		}
+
 	private:
 		Audio::FPatchOutputStrongPtr AudioBusPatch;
 		Audio::FAlignedFloatBuffer TempBuffer;
@@ -207,13 +220,13 @@ namespace AudioModulation
 		float Gain = 1.0f;
 		uint8 bBypass : 1;
 		uint8 bInvert : 1;
-		uint8 bInitialized : 1;
+		uint8 bBusRequiresPatch : 1;
 	};
 } // namespace AudioModulation
 
-AudioModulation::FGeneratorPtr USoundModulationGeneratorEnvelopeFollower::CreateInstance(Audio::FDeviceId InDeviceId) const
+AudioModulation::FGeneratorPtr USoundModulationGeneratorEnvelopeFollower::CreateInstance() const
 {
 	using namespace AudioModulation;
 
-	return FGeneratorPtr(new FEnvelopeFollowerGenerator(Params, InDeviceId));
+	return FGeneratorPtr(new FEnvelopeFollowerGenerator(Params));
 }
diff --git a/Engine/Plugins/Runtime/AudioModulation/Source/AudioModulation/Private/Generators/SoundModulationLFO.cpp b/Engine/Plugins/Runtime/AudioModulation/Source/AudioModulation/Private/Generators/SoundModulationLFO.cpp
--- a/Engine/Plugins/Runtime/AudioModulation/Source/AudioModulation/Private/Generators/SoundModulationLFO.cpp
+++ b/Engine/Plugins/Runtime/AudioModulation/Source/AudioModulation/Private/Generators/SoundModulationLFO.cpp
@@ -30,7 +30,9 @@ namespace AudioModulation
 			FLFOGenerator() = default;
 
 			FLFOGenerator(const FLFOGenerator& InGenerator)
-				: IGenerator(InGenerator.AudioDeviceId)
+				: LFO(InGenerator.LFO)
+				, Value(InGenerator.Value)
+				, Params(InGenerator.Params)
 			{
 			}
 
@@ -167,7 +169,7 @@ namespace AudioModulation
 	};
 } // namespace AudioModulation
 
-AudioModulation::FGeneratorPtr USoundModulationGeneratorLFO::CreateInstance(Audio::FDeviceId InDeviceId) const
+AudioModulation::FGeneratorPtr USoundModulationGeneratorLFO::CreateInstance() const
 {
 	using namespace AudioModulation;
 
diff --git a/Engine/Plugins/Runtime/AudioModulation/Source/AudioModulation/Private/MetaSoundNodes/GetModulatorValueNode.cpp b/Engine/Plugins/Runtime/AudioModulation/Source/AudioModulation/Private/MetaSoundNodes/GetModulatorValueNode.cpp
--- a/Engine/Plugins/Runtime/AudioModulation/Source/AudioModulation/Private/MetaSoundNodes/GetModulatorValueNode.cpp
+++ b/Engine/Plugins/Runtime/AudioModulation/Source/AudioModulation/Private/MetaSoundNodes/GetModulatorValueNode.cpp
@@ -1,9 +1,13 @@
 // Copyright Epic Games, Inc. All Rights Reserved.
+
+#include "AudioDefines.h"
+#include "AudioModulation.h"
 #include "MetasoundDataFactory.h"
 #include "MetasoundExecutableOperator.h"
 #include "MetasoundFacade.h"
 #include "MetasoundNodeRegistrationMacro.h"
 #include "MetasoundPrimitives.h"
+#include "MetasoundSourceInterface.h"
 #include "SoundModulatorAsset.h"
 
 #define LOCTEXT_NAMESPACE "AudioModulationNodes"
@@ -64,13 +68,14 @@ namespace AudioModulation
 			FSoundModulatorAssetReadRef ModulatorReadRef = InputCollection.GetDataReadReferenceOrConstruct<FSoundModulatorAsset>("Modulator");
 			FBoolReadRef NormalizedReadRef = InputCollection.GetDataReadReferenceOrConstructWithVertexDefault<bool>(InputInterface, "Normalized", InParams.OperatorSettings);
 
-			return MakeUnique<FGetModulatorValueNodeOperator>(InParams.OperatorSettings, ModulatorReadRef, NormalizedReadRef);
+			return MakeUnique<FGetModulatorValueNodeOperator>(InParams, ModulatorReadRef, NormalizedReadRef);
 		}
 
-		FGetModulatorValueNodeOperator(const Metasound::FOperatorSettings& InSettings, const FSoundModulatorAssetReadRef& InModulator, const Metasound::FBoolReadRef& InNormalized)
-			: Modulator(InModulator)
+		FGetModulatorValueNodeOperator(const Metasound::FCreateOperatorParams& InParams, const FSoundModulatorAssetReadRef& InModulator, const Metasound::FBoolReadRef& InNormalized)
+			: DeviceId(InParams.Environment.GetValue<Audio::FDeviceId>(Metasound::Frontend::SourceInterface::Environment::DeviceID))
+			, Modulator(InModulator)
 			, Normalized(InNormalized)
-			, OutValue(Metasound::TDataWriteReferenceFactory<float>::CreateAny(InSettings))
+			, OutValue(Metasound::TDataWriteReferenceFactory<float>::CreateAny(InParams.OperatorSettings))
 		{
 		}
 
@@ -100,28 +105,39 @@ namespace AudioModulation
 		void Execute()
 		{
 			const FSoundModulatorAsset& ModulatorAsset = *Modulator;
-			if (ModulatorAsset.IsValid())
-			{
-				float Value = ModulatorAsset->GetValue();
 
-				if (!*Normalized)
+			// Checks if handle needs refresh.  This can happen if the incoming modulator is changed or nulled.
+			if (ModulatorAsset.GetModulatorId() != ModHandle.GetModulatorId())
+			{
+				if (IAudioModulationManager* InModulation = AudioModulation::GetDeviceModulationManager(DeviceId))
 				{
-					const Audio::FModulationParameter& Parameter = ModulatorAsset->GetParameter();
-					if (Parameter.bRequiresConversion)
-					{
-						Parameter.UnitFunction(Value);
-					}
+					ModHandle = ModulatorAsset->CreateModulatorHandle(*InModulation);
 				}
+			}
 
-				*OutValue = Value;
+			float Value = 1.0f;
+			if (ModHandle.IsValid())
+			{
+				ModHandle.GetValueThreadSafe(Value);
 			}
-			else
+
+			if (!*Normalized)
 			{
-				*OutValue = 1.0f;
+				const Audio::FModulationParameter& Parameter = ModHandle.GetParameter();
+				if (Parameter.bRequiresConversion)
+				{
+					Parameter.UnitFunction(Value);
+				}
 			}
+
+			*OutValue = Value;
 		}
 
 	private:
+		Audio::FDeviceId DeviceId;
+
+		Audio::FModulatorHandle ModHandle;
+
 		FSoundModulatorAssetReadRef Modulator;
 		Metasound::FBoolReadRef Normalized;
 		Metasound::TDataWriteReference<float> OutValue;
diff --git a/Engine/Plugins/Runtime/AudioModulation/Source/AudioModulation/Private/MetaSoundNodes/MixModulatorsNode.cpp b/Engine/Plugins/Runtime/AudioModulation/Source/AudioModulation/Private/MetaSoundNodes/MixModulatorsNode.cpp
--- a/Engine/Plugins/Runtime/AudioModulation/Source/AudioModulation/Private/MetaSoundNodes/MixModulatorsNode.cpp
+++ b/Engine/Plugins/Runtime/AudioModulation/Source/AudioModulation/Private/MetaSoundNodes/MixModulatorsNode.cpp
@@ -1,10 +1,13 @@
 // Copyright Epic Games, Inc. All Rights Reserved.
+
+#include "AudioModulation.h"
 #include "IAudioModulation.h"
 #include "MetasoundDataFactory.h"
 #include "MetasoundExecutableOperator.h"
 #include "MetasoundFacade.h"
 #include "MetasoundNodeRegistrationMacro.h"
 #include "MetasoundPrimitives.h"
+#include "MetasoundSourceInterface.h"
 #include "SoundModulatorAsset.h"
 
 #define LOCTEXT_NAMESPACE "AudioModulationNodes"
@@ -69,20 +72,21 @@ namespace AudioModulation
 			FSoundModulationParameterAssetReadRef ParameterReadRef = InputCollection.GetDataReadReferenceOrConstruct<FSoundModulationParameterAsset>("MixParameter");
 			FBoolReadRef NormalizedReadRef = InputCollection.GetDataReadReferenceOrConstructWithVertexDefault<bool>(InputInterface, "Normalized", InParams.OperatorSettings);
 
-			return MakeUnique<FMixModulatorsNodeOperator>(InParams.OperatorSettings, Modulator1ReadRef, Modulator2ReadRef, ParameterReadRef, NormalizedReadRef);
+			return MakeUnique<FMixModulatorsNodeOperator>(InParams, Modulator1ReadRef, Modulator2ReadRef, ParameterReadRef, NormalizedReadRef);
 		}
 
 		FMixModulatorsNodeOperator(
-			const Metasound::FOperatorSettings& InSettings,
+			const Metasound::FCreateOperatorParams& InParams,
 			const FSoundModulatorAssetReadRef& InModulator1,
 			const FSoundModulatorAssetReadRef& InModulator2,
 			const FSoundModulationParameterAssetReadRef& InParameter,
 			const Metasound::FBoolReadRef& InNormalized)
-			: Modulator1(InModulator1)
+			: DeviceId(InParams.Environment.GetValue<Audio::FDeviceId>(Metasound::Frontend::SourceInterface::Environment::DeviceID))
+			, Modulator1(InModulator1)
 			, Modulator2(InModulator2)
 			, Normalized(InNormalized)
 			, Parameter(InParameter)
-			, OutValue(Metasound::TDataWriteReferenceFactory<float>::CreateAny(InSettings))
+			, OutValue(Metasound::TDataWriteReferenceFactory<float>::CreateAny(InParams.OperatorSettings))
 		{
 		}
 
@@ -113,19 +117,30 @@ namespace AudioModulation
 
 		void Execute()
 		{
-			float Value1 = 1.0f;
-			const FSoundModulatorAsset& ModulatorAsset1 = *Modulator1;
-			if (ModulatorAsset1.IsValid())
+			auto UpdateValue = [](const FSoundModulatorAssetReadRef& ReadRef, Audio::FDeviceId InDeviceId, Audio::FModulatorHandle& OutHandle)
 			{
-				Value1 = ModulatorAsset1->GetValue();
-			}
+				float Value = 1.0f;
+				const FSoundModulatorAsset& ModulatorAsset = *ReadRef;
 
-			float Value2 = 1.0f;
-			const FSoundModulatorAsset& ModulatorAsset2 = *Modulator2;
-			if (ModulatorAsset2.IsValid())
-			{
-				Value2 = ModulatorAsset2->GetValue();
-			}
+				// Checks if handle needs refresh.  This can happen if the incoming modulator is changed or nulled.
+				if (ModulatorAsset.GetModulatorId() != OutHandle.GetModulatorId())
+				{
+					if (IAudioModulationManager* InModulation = AudioModulation::GetDeviceModulationManager(InDeviceId))
+					{
+						OutHandle = ModulatorAsset->CreateModulatorHandle(*InModulation);
+					}
+				}
+
+				if (OutHandle.IsValid())
+				{
+					Value = OutHandle.GetValueThreadSafe(Value);
+				}
+
+				return Value;
+			};
+
+			float Value1 = UpdateValue(Modulator1, DeviceId, ModHandle1);
+			const float Value2 = UpdateValue(Modulator2, DeviceId, ModHandle2);
 
 			const FSoundModulationParameterAsset& ParameterAsset = *Parameter;
 			if (ParameterAsset.IsValid())
@@ -141,6 +156,11 @@ namespace AudioModulation
 		}
 
 	private:
+		const Audio::FDeviceId DeviceId;
+
+		Audio::FModulatorHandle ModHandle1;
+		Audio::FModulatorHandle ModHandle2;
+
 		FSoundModulatorAssetReadRef Modulator1;
 		FSoundModulatorAssetReadRef Modulator2;
 		Metasound::FBoolReadRef Normalized;
diff --git a/Engine/Plugins/Runtime/AudioModulation/Source/AudioModulation/Private/SoundControlBus.cpp b/Engine/Plugins/Runtime/AudioModulation/Source/AudioModulation/Private/SoundControlBus.cpp
--- a/Engine/Plugins/Runtime/AudioModulation/Source/AudioModulation/Private/SoundControlBus.cpp
+++ b/Engine/Plugins/Runtime/AudioModulation/Source/AudioModulation/Private/SoundControlBus.cpp
@@ -10,7 +10,6 @@
 #include "DSP/BufferVectorOperations.h"
 #include "Engine/World.h"
 #include "SoundControlBusProxy.h"
-#include "SoundModulatorAssetProxy.h"
 
 
 USoundControlBus::USoundControlBus(const FObjectInitializer& ObjectInitializer)
@@ -23,6 +22,11 @@ USoundControlBus::USoundControlBus(const FObjectInitializer& ObjectInitializer)
 {
 }
 
+TUniquePtr<Audio::IModulatorSettings> USoundControlBus::CreateProxySettings() const
+{
+	return TUniquePtr<Audio::IModulatorSettings>(new AudioModulation::FControlBusSettings(*this));
+}
+
 #if WITH_EDITOR
 void USoundControlBus::PostDuplicate(EDuplicateMode::Type DuplicateMode)
 {
@@ -69,7 +73,7 @@ void USoundControlBus::PostEditChangeProperty(FPropertyChangedEvent& InPropertyC
 			}
 		}
 
-		AudioModulation::IterateModulationImpl([this](AudioModulation::FAudioModulation& OutModSystem)
+		AudioModulation::IterateModulationManagers([this](AudioModulation::FAudioModulationManager& OutModSystem)
 		{
 			OutModSystem.UpdateModulator(*this);
 		});
@@ -109,9 +113,9 @@ void USoundControlBus::BeginDestroy()
 		if (AudioDevice.IsValid())
 		{
 			check(AudioDevice->IsModulationPluginEnabled());
-			if (IAudioModulation* ModulationInterface = AudioDevice->ModulationInterface.Get())
+			if (IAudioModulationManager* ModulationInterface = AudioDevice->ModulationInterface.Get())
 			{
-				FAudioModulation* Modulation = static_cast<FAudioModulation*>(ModulationInterface);
+				FAudioModulationManager* Modulation = static_cast<FAudioModulationManager*>(ModulationInterface);
 				check(Modulation);
 				Modulation->DeactivateBus(*this);
 			}
@@ -132,5 +136,11 @@ const Audio::FModulationMixFunction USoundControlBus::GetMixFunction() const
 TUniquePtr<Audio::IProxyData> USoundControlBus::CreateNewProxyData(const Audio::FProxyDataInitParams& InitParams)
 {
 	using namespace AudioModulation;
-	return MakeUnique<TSoundModulatorAssetProxy<USoundControlBus>>(*this);
+	return MakeUnique<FSoundModulatorAssetProxy>(*this);
 }
+
+const Audio::FModulationParameter& USoundControlBus::GetOutputParameter() const
+{
+	return AudioModulation::GetOrRegisterParameter(Parameter, *this);
+}
+
diff --git a/Engine/Plugins/Runtime/AudioModulation/Source/AudioModulation/Private/SoundControlBusMix.cpp b/Engine/Plugins/Runtime/AudioModulation/Source/AudioModulation/Private/SoundControlBusMix.cpp
--- a/Engine/Plugins/Runtime/AudioModulation/Source/AudioModulation/Private/SoundControlBusMix.cpp
+++ b/Engine/Plugins/Runtime/AudioModulation/Source/AudioModulation/Private/SoundControlBusMix.cpp
@@ -48,9 +48,9 @@ void USoundControlBusMix::BeginDestroy()
 		{
 			if (AudioDevice->IsModulationPluginEnabled())
 			{
-				if (IAudioModulation* ModulationInterface = AudioDevice->ModulationInterface.Get())
+				if (IAudioModulationManager* ModulationInterface = AudioDevice->ModulationInterface.Get())
 				{
-					FAudioModulation* Modulation = static_cast<FAudioModulation*>(ModulationInterface);
+					FAudioModulationManager* Modulation = static_cast<FAudioModulationManager*>(ModulationInterface);
 					check(Modulation);
 					Modulation->DeactivateBusMix(*this);
 				}
@@ -61,7 +61,7 @@ void USoundControlBusMix::BeginDestroy()
 
 void USoundControlBusMix::ActivateMix()
 {
-	AudioModulation::IterateModulationImpl([this](AudioModulation::FAudioModulation& OutModulation)
+	AudioModulation::IterateModulationManagers([this](AudioModulation::FAudioModulationManager& OutModulation)
 	{
 		OutModulation.ActivateBusMix(*this);
 	});
@@ -69,7 +69,7 @@ void USoundControlBusMix::ActivateMix()
 
 void USoundControlBusMix::DeactivateMix()
 {
-	AudioModulation::IterateModulationImpl([this](AudioModulation::FAudioModulation& OutModulation)
+	AudioModulation::IterateModulationManagers([this](AudioModulation::FAudioModulationManager& OutModulation)
 	{
 		OutModulation.DeactivateBusMix(*this);
 	});
@@ -77,7 +77,7 @@ void USoundControlBusMix::DeactivateMix()
 
 void USoundControlBusMix::DeactivateAllMixes()
 {
-	AudioModulation::IterateModulationImpl([this](AudioModulation::FAudioModulation& OutModulation)
+	AudioModulation::IterateModulationManagers([this](AudioModulation::FAudioModulationManager& OutModulation)
 	{
 		OutModulation.DeactivateAllBusMixes();
 	});
@@ -147,7 +147,7 @@ void USoundControlBusMix::OnPropertyChanged(FProperty* Property, EPropertyChange
 		}
 	}
 
-	AudioModulation::IterateModulationImpl([this](AudioModulation::FAudioModulation& OutModulation)
+	AudioModulation::IterateModulationManagers([this](AudioModulation::FAudioModulationManager& OutModulation)
 	{
 		OutModulation.UpdateMix(*this);
 	});
@@ -175,7 +175,7 @@ void USoundControlBusMix::SaveMixToProfile()
 
 void USoundControlBusMix::SoloMix()
 {
-	AudioModulation::IterateModulationImpl([this](AudioModulation::FAudioModulation& OutModulation)
+	AudioModulation::IterateModulationManagers([this](AudioModulation::FAudioModulationManager& OutModulation)
 	{
 		OutModulation.SoloBusMix(*this);
 	});
diff --git a/Engine/Plugins/Runtime/AudioModulation/Source/AudioModulation/Private/SoundControlBusMixProxy.cpp b/Engine/Plugins/Runtime/AudioModulation/Source/AudioModulation/Private/SoundControlBusMixProxy.cpp
--- a/Engine/Plugins/Runtime/AudioModulation/Source/AudioModulation/Private/SoundControlBusMixProxy.cpp
+++ b/Engine/Plugins/Runtime/AudioModulation/Source/AudioModulation/Private/SoundControlBusMixProxy.cpp
@@ -19,13 +19,13 @@ namespace AudioModulation
 {
 	const FBusMixId InvalidBusMixId = INDEX_NONE;
 
-	FModulatorBusMixStageSettings::FModulatorBusMixStageSettings(const FSoundControlBusMixStage& InStage, Audio::FDeviceId InDeviceId)
+	FModulatorBusMixStageSettings::FModulatorBusMixStageSettings(const FSoundControlBusMixStage& InStage)
 		: TModulatorBase<FBusId>(InStage.Bus->GetName(), InStage.Bus->GetUniqueID())
 		, Address(InStage.Bus->Address)
 		, ParamClassId(INDEX_NONE)
 		, ParamId(INDEX_NONE)
 		, Value(InStage.Value)
-		, BusSettings(FControlBusSettings(*InStage.Bus, InDeviceId))
+		, BusSettings(FControlBusSettings(*InStage.Bus))
 	{
 		if (USoundModulationParameter* Parameter = InStage.Bus->Parameter)
 		{
@@ -49,14 +49,14 @@ namespace AudioModulation
 	{
 	}
 
-	FModulatorBusMixSettings::FModulatorBusMixSettings(const USoundControlBusMix& InBusMix, Audio::FDeviceId InDeviceId)
+	FModulatorBusMixSettings::FModulatorBusMixSettings(const USoundControlBusMix& InBusMix)
 		: TModulatorBase<FBusMixId>(InBusMix.GetName(), InBusMix.GetUniqueID())
 	{
 		for (const FSoundControlBusMixStage& Stage : InBusMix.MixStages)
 		{
 			if (Stage.Bus)
 			{
-				Stages.Add(FModulatorBusMixStageSettings(Stage, InDeviceId));
+				Stages.Add(FModulatorBusMixStageSettings(Stage));
 			}
 			else
 			{
diff --git a/Engine/Plugins/Runtime/AudioModulation/Source/AudioModulation/Private/SoundControlBusMixProxy.h b/Engine/Plugins/Runtime/AudioModulation/Source/AudioModulation/Private/SoundControlBusMixProxy.h
--- a/Engine/Plugins/Runtime/AudioModulation/Source/AudioModulation/Private/SoundControlBusMixProxy.h
+++ b/Engine/Plugins/Runtime/AudioModulation/Source/AudioModulation/Private/SoundControlBusMixProxy.h
@@ -24,7 +24,7 @@ namespace AudioModulation
 	class FModulatorBusMixStageSettings : public TModulatorBase<FBusId>
 	{
 	public:
-		FModulatorBusMixStageSettings(const FSoundControlBusMixStage& InStage, Audio::FDeviceId InDeviceId);
+		FModulatorBusMixStageSettings(const FSoundControlBusMixStage& InStage);
 
 		FString Address;
 		uint32 ParamClassId = INDEX_NONE;
@@ -36,7 +36,7 @@ namespace AudioModulation
 	class FModulatorBusMixSettings : public TModulatorBase<FBusMixId>
 	{
 	public:
-		FModulatorBusMixSettings(const USoundControlBusMix& InBusMix, Audio::FDeviceId InDeviceId);
+		FModulatorBusMixSettings(const USoundControlBusMix& InBusMix);
 		FModulatorBusMixSettings(FModulatorBusMixSettings&& InSettings);
 
 		TArray<FModulatorBusMixStageSettings> Stages;
diff --git a/Engine/Plugins/Runtime/AudioModulation/Source/AudioModulation/Private/SoundControlBusProxy.cpp b/Engine/Plugins/Runtime/AudioModulation/Source/AudioModulation/Private/SoundControlBusProxy.cpp
--- a/Engine/Plugins/Runtime/AudioModulation/Source/AudioModulation/Private/SoundControlBusProxy.cpp
+++ b/Engine/Plugins/Runtime/AudioModulation/Source/AudioModulation/Private/SoundControlBusProxy.cpp
@@ -14,6 +14,12 @@ namespace AudioModulation
 {
 	const FBusId InvalidBusId = INDEX_NONE;
 
+	Audio::FModulatorTypeId FControlBusSettings::Register(Audio::FModulatorHandleId HandleId, IAudioModulationManager& InModulation) const
+	{
+		FAudioModulationSystem& ModSystem = static_cast<FAudioModulationManager&>(InModulation).GetSystem();
+		return ModSystem.RegisterModulator(HandleId, *this);
+	}
+
 	FControlBusProxy::FControlBusProxy()
 		: DefaultValue(0.0f)
 		, GeneratorValue(1.0f)
diff --git a/Engine/Plugins/Runtime/AudioModulation/Source/AudioModulation/Private/SoundControlBusProxy.h b/Engine/Plugins/Runtime/AudioModulation/Source/AudioModulation/Private/SoundControlBusProxy.h
--- a/Engine/Plugins/Runtime/AudioModulation/Source/AudioModulation/Private/SoundControlBusProxy.h
+++ b/Engine/Plugins/Runtime/AudioModulation/Source/AudioModulation/Private/SoundControlBusProxy.h
@@ -17,29 +17,48 @@ namespace AudioModulation
 	using FBusId = uint32;
 	extern const FBusId InvalidBusId;
 
-	struct FControlBusSettings : public TModulatorBase<FBusId>
+	struct FControlBusSettings : public TModulatorBase<FBusId>, public Audio::IModulatorSettings
 	{
 		bool bBypass;
 		float DefaultValue;
 
 		TArray<FModulationGeneratorSettings> GeneratorSettings;
 		Audio::FModulationMixFunction MixFunction;
+		Audio::FModulationParameter OutputParameter;
 
-		FControlBusSettings(const USoundControlBus& InBus, Audio::FDeviceId InDeviceId)
+		FControlBusSettings(const USoundControlBus& InBus)
 			: TModulatorBase<FBusId>(InBus.GetName(), InBus.GetUniqueID())
 			, bBypass(InBus.bBypass)
 			, DefaultValue(InBus.GetDefaultNormalizedValue())
 			, MixFunction(InBus.GetMixFunction())
+			, OutputParameter(InBus.GetOutputParameter())
 		{
 			for (const USoundModulationGenerator* Generator : InBus.Generators)
 			{
 				if (Generator)
 				{
-					FModulationGeneratorSettings Settings(*Generator, InDeviceId);
+					FModulationGeneratorSettings Settings(*Generator);
 					GeneratorSettings.Add(MoveTemp(Settings));
 				}
 			}
 		}
+
+		virtual TUniquePtr<IModulatorSettings> Clone() const override
+		{
+			return TUniquePtr<IModulatorSettings>(new FControlBusSettings(*this));
+		}
+
+		virtual Audio::FModulatorId GetModulatorId() const override
+		{
+			return static_cast<Audio::FModulatorId>(GetId());
+		}
+
+		virtual const Audio::FModulationParameter& GetOutputParameter() const override
+		{
+			return OutputParameter;
+		}
+
+		virtual Audio::FModulatorTypeId Register(Audio::FModulatorHandleId HandleId, IAudioModulationManager& InModulation) const override;
 	};
 
 	class FControlBusProxy : public TModulatorProxyRefType<FBusId, FControlBusProxy, FControlBusSettings>
diff --git a/Engine/Plugins/Runtime/AudioModulation/Source/AudioModulation/Private/SoundModulationGenerator.cpp b/Engine/Plugins/Runtime/AudioModulation/Source/AudioModulation/Private/SoundModulationGenerator.cpp
--- a/Engine/Plugins/Runtime/AudioModulation/Source/AudioModulation/Private/SoundModulationGenerator.cpp
+++ b/Engine/Plugins/Runtime/AudioModulation/Source/AudioModulation/Private/SoundModulationGenerator.cpp
@@ -6,7 +6,7 @@
 #include "AudioModulationLogging.h"
 #include "AudioModulationSystem.h"
 #include "Engine/World.h"
-#include "SoundModulatorAssetProxy.h"
+#include "SoundModulationGeneratorProxy.h"
 #include "Templates/Function.h"
 
 
@@ -34,7 +34,7 @@ void USoundModulationGenerator::PostEditChangeProperty(FPropertyChangedEvent& In
 	// Guards against slamming the modulation system with changes when using sliders.
 	if (InPropertyChangedEvent.ChangeType != EPropertyChangeType::Interactive)
 	{
-		AudioModulation::IterateModulationImpl([this](AudioModulation::FAudioModulation& OutModulation)
+		AudioModulation::IterateModulationManagers([this](AudioModulation::FAudioModulationManager& OutModulation)
 		{
 			OutModulation.UpdateModulator(*this);
 		});
@@ -47,7 +47,13 @@ void USoundModulationGenerator::PostEditChangeProperty(FPropertyChangedEvent& In
 TUniquePtr<Audio::IProxyData> USoundModulationGenerator::CreateNewProxyData(const Audio::FProxyDataInitParams& InitParams)
 {
 	using namespace AudioModulation;
-	return MakeUnique<TSoundModulatorAssetProxy<USoundModulationGenerator>>(*this);
+	return MakeUnique<FSoundModulatorAssetProxy>(*this);
+}
+
+TUniquePtr<Audio::IModulatorSettings> USoundModulationGenerator::CreateProxySettings() const
+{
+	using namespace AudioModulation;
+	return TUniquePtr<Audio::IModulatorSettings>(new FModulationGeneratorSettings(*this));
 }
 
 void USoundModulationGenerator::BeginDestroy()
@@ -60,9 +66,9 @@ void USoundModulationGenerator::BeginDestroy()
 		if (AudioDevice.IsValid())
 		{
 			check(AudioDevice->IsModulationPluginEnabled());
-			if (IAudioModulation* ModulationInterface = AudioDevice->ModulationInterface.Get())
+			if (IAudioModulationManager* ModulationInterface = AudioDevice->ModulationInterface.Get())
 			{
-				FAudioModulation* Modulation = static_cast<FAudioModulation*>(ModulationInterface);
+				FAudioModulationManager* Modulation = static_cast<FAudioModulationManager*>(ModulationInterface);
 				check(Modulation);
 				Modulation->DeactivateGenerator(*this);
 			}
diff --git a/Engine/Plugins/Runtime/AudioModulation/Source/AudioModulation/Private/SoundModulationGeneratorProxy.cpp b/Engine/Plugins/Runtime/AudioModulation/Source/AudioModulation/Private/SoundModulationGeneratorProxy.cpp
--- a/Engine/Plugins/Runtime/AudioModulation/Source/AudioModulation/Private/SoundModulationGeneratorProxy.cpp
+++ b/Engine/Plugins/Runtime/AudioModulation/Source/AudioModulation/Private/SoundModulationGeneratorProxy.cpp
@@ -9,10 +9,23 @@ namespace AudioModulation
 {
 	const FGeneratorId InvalidGeneratorId = INDEX_NONE;
 
+	const Audio::FModulationParameter& FModulationGeneratorSettings::GetOutputParameter() const
+	{
+		// For now, all generators are normalized values (0, 1).
+		return Audio::GetModulationParameter({ });
+	}
+
+	Audio::FModulatorTypeId FModulationGeneratorSettings::Register(Audio::FModulatorHandleId HandleId, IAudioModulationManager& InModulation) const
+	{
+		FAudioModulationSystem& ModSystem = static_cast<FAudioModulationManager&>(InModulation).GetSystem();
+		return ModSystem.RegisterModulator(HandleId, *this);
+	}
+
 	FModulatorGeneratorProxy::FModulatorGeneratorProxy(FModulationGeneratorSettings&& InSettings, FAudioModulationSystem& InModSystem)
 		: TModulatorProxyRefType(InSettings.GetName(), InSettings.GetId(), InModSystem)
 		, Generator(MoveTemp(InSettings.Generator))
 	{
+		Generator->Init(InModSystem.GetAudioDeviceId());
 	}
 
 	FModulatorGeneratorProxy& FModulatorGeneratorProxy::operator =(FModulationGeneratorSettings&& InSettings)
@@ -28,5 +41,4 @@ namespace AudioModulation
 
 		return *this;
 	}
-
 } // namespace AudioModulation
diff --git a/Engine/Plugins/Runtime/AudioModulation/Source/AudioModulation/Private/SoundModulationGeneratorProxy.h b/Engine/Plugins/Runtime/AudioModulation/Source/AudioModulation/Private/SoundModulationGeneratorProxy.h
--- a/Engine/Plugins/Runtime/AudioModulation/Source/AudioModulation/Private/SoundModulationGeneratorProxy.h
+++ b/Engine/Plugins/Runtime/AudioModulation/Source/AudioModulation/Private/SoundModulationGeneratorProxy.h
@@ -1,7 +1,9 @@
 // Copyright Epic Games, Inc. All Rights Reserved.
 #pragma once
 
+#include "AudioDefines.h"
 #include "AudioDeviceManager.h"
+#include "IAudioModulation.h"
 #include "SoundModulationGenerator.h"
 #include "SoundModulationProxy.h"
 
@@ -21,15 +23,15 @@ namespace AudioModulation
 	using FGeneratorProxyMap = TMap<FGeneratorId, FModulatorGeneratorProxy>;
 	using FGeneratorHandle = TProxyHandle<FGeneratorId, FModulatorGeneratorProxy, FModulationGeneratorSettings>;
 
-	struct FModulationGeneratorSettings : public TModulatorBase<FGeneratorId>
+	struct FModulationGeneratorSettings : public TModulatorBase<FGeneratorId>, public Audio::IModulatorSettings
 	{
 		FGeneratorPtr Generator;
 
 		FModulationGeneratorSettings() = default;
 
-		FModulationGeneratorSettings(const USoundModulationGenerator& InGenerator, Audio::FDeviceId InDeviceId)
+		FModulationGeneratorSettings(const USoundModulationGenerator& InGenerator)
 			: TModulatorBase<FGeneratorId>(InGenerator.GetName(), InGenerator.GetUniqueID())
-			, Generator(InGenerator.CreateInstance(InDeviceId))
+			, Generator(InGenerator.CreateInstance())
 		{
 		}
 
@@ -44,6 +46,20 @@ namespace AudioModulation
 			, Generator(MoveTemp(InSettings.Generator))
 		{
 		}
+
+		virtual TUniquePtr<IModulatorSettings> Clone() const override
+		{
+			return TUniquePtr<IModulatorSettings>(new FModulationGeneratorSettings(*this));
+		}
+
+
+		virtual Audio::FModulatorId GetModulatorId() const override
+		{
+			return static_cast<Audio::FModulatorId>(GetId());
+		}
+
+		virtual const Audio::FModulationParameter& GetOutputParameter() const override;
+		virtual Audio::FModulatorTypeId Register(Audio::FModulatorHandleId HandleId, IAudioModulationManager& InModulation) const override;
 	};
 
 	class FModulatorGeneratorProxy : public TModulatorProxyRefType<FGeneratorId, FModulatorGeneratorProxy, FModulationGeneratorSettings>
@@ -61,6 +77,11 @@ namespace AudioModulation
 			return Generator->GetValue();
 		}
 
+		void Init(Audio::FDeviceId InDeviceId)
+		{
+			Generator->Init(InDeviceId);
+		}
+
 		bool IsBypassed() const
 		{
 			return Generator->IsBypassed();
diff --git a/Engine/Plugins/Runtime/AudioModulation/Source/AudioModulation/Private/SoundModulationParameter.cpp b/Engine/Plugins/Runtime/AudioModulation/Source/AudioModulation/Private/SoundModulationParameter.cpp
--- a/Engine/Plugins/Runtime/AudioModulation/Source/AudioModulation/Private/SoundModulationParameter.cpp
+++ b/Engine/Plugins/Runtime/AudioModulation/Source/AudioModulation/Private/SoundModulationParameter.cpp
@@ -3,6 +3,9 @@
 
 #include "AudioDevice.h"
 #include "AudioDeviceManager.h"
+#include "AudioModulation.h"
+#include "AudioModulationLogging.h"
+#include "IAudioModulation.h"
 
 
 TUniquePtr<Audio::IProxyData> USoundModulationParameter::CreateNewProxyData(const Audio::FProxyDataInitParams& InitParams)
@@ -33,6 +36,25 @@ void USoundModulationParameter::RefreshUnitValue()
 }
 #endif // WITH_EDITOR
 
+Audio::FModulationParameter USoundModulationParameter::CreateParameter() const
+{
+	Audio::FModulationParameter Parameter;
+	Parameter.ParameterName = GetFName();
+	Parameter.bRequiresConversion = RequiresUnitConversion();
+	Parameter.MixFunction = GetMixFunction();
+	Parameter.UnitFunction = GetUnitConversionFunction();
+	Parameter.NormalizedFunction = GetNormalizedConversionFunction();
+	Parameter.DefaultValue = GetUnitDefault();
+	Parameter.MinValue = GetUnitMin();
+	Parameter.MaxValue = GetUnitMax();
+
+#if WITH_EDITORONLY_DATA
+	Parameter.UnitDisplayName = Settings.UnitDisplayName;
+#endif // WITH_EDITORONLY_DATA
+
+	return Parameter;
+}
+
 bool USoundModulationParameterFrequencyBase::RequiresUnitConversion() const
 {
 	return true;
@@ -202,34 +224,40 @@ float USoundModulationParameterVolume::GetUnitMax() const
 
 namespace AudioModulation
 {
-	FSoundModulationPluginParameterAssetProxy::FSoundModulationPluginParameterAssetProxy(USoundModulationParameter* InParameter)
+	const Audio::FModulationParameter& GetOrRegisterParameter(const USoundModulationParameter* InParameter, const USoundModulatorBase& InReferencingModulator)
 	{
-		using namespace Audio;
-
-		if (!InParameter || !GEngine)
-		{
-			return;
-		}
-
-		FAudioDeviceHandle AudioDevice;
-		if (UWorld* World = GEngine->GetWorldFromContextObject(InParameter, EGetWorldErrorMode::ReturnNull))
+		FName ParamName;
+		if (InParameter)
 		{
-			if (!World->bAllowAudioPlayback || World->IsNetMode(NM_DedicatedServer))
+			ParamName = InParameter->GetFName();
+			if (!Audio::IsModulationParameterRegistered(ParamName))
 			{
-				return;
+				UE_LOG(LogAudioModulation, Display,
+					TEXT("Parameter '%s' not registered.  Registration forced via activating referrencing modulator '%s'."),
+					*ParamName.ToString(),
+					*InReferencingModulator.GetName());
 			}
 
-			AudioDevice = World->GetAudioDevice();
-		}
-		else
-		{
-			AudioDevice = GEngine->GetMainAudioDevice();
+			Audio::RegisterModulationParameter(ParamName, InParameter->CreateParameter());
 		}
 
-		if (IAudioModulation* Modulation = AudioDevice->ModulationInterface.Get())
+		// Returns default modulation parameter if no parameter provided.
+		return Audio::GetModulationParameter(ParamName);
+	}
+
+	FSoundModulationPluginParameterAssetProxy::FSoundModulationPluginParameterAssetProxy(USoundModulationParameter* InParameter)
+	{
+		using namespace Audio;
+
+		if (InParameter)
 		{
 			const FName ParameterName = InParameter->GetFName();
-			Parameter = Modulation->GetParameter(ParameterName);
+			Parameter = Audio::GetModulationParameter(ParameterName);
+
+			if (ParameterName != Parameter.ParameterName)
+			{
+				UE_LOG(LogAudioModulation, Log, TEXT("Parameter '%s' not registered. Using default parameter."), *ParameterName.ToString());
+			}
 		}
 	}
 
diff --git a/Engine/Plugins/Runtime/AudioModulation/Source/AudioModulation/Private/SoundModulationPatch.cpp b/Engine/Plugins/Runtime/AudioModulation/Source/AudioModulation/Private/SoundModulationPatch.cpp
--- a/Engine/Plugins/Runtime/AudioModulation/Source/AudioModulation/Private/SoundModulationPatch.cpp
+++ b/Engine/Plugins/Runtime/AudioModulation/Source/AudioModulation/Private/SoundModulationPatch.cpp
@@ -7,8 +7,9 @@
 #include "AudioModulationStatics.h"
 #include "AudioModulationSystem.h"
 #include "SoundControlBus.h"
+#include "SoundModulationParameter.h"
+#include "SoundModulationPatchProxy.h"
 #include "SoundModulationTransform.h"
-#include "SoundModulatorAssetProxy.h"
 
 
 #define LOCTEXT_NAMESPACE "SoundModulationPatch"
@@ -22,13 +23,24 @@ USoundModulationPatch::USoundModulationPatch(const FObjectInitializer& ObjectIni
 TUniquePtr<Audio::IProxyData> USoundModulationPatch::CreateNewProxyData(const Audio::FProxyDataInitParams& InitParams)
 {
 	using namespace AudioModulation;
-	return MakeUnique<TSoundModulatorAssetProxy<USoundModulationPatch>>(*this);
+	return MakeUnique<FSoundModulatorAssetProxy>(*this);
+}
+
+TUniquePtr<Audio::IModulatorSettings> USoundModulationPatch::CreateProxySettings() const
+{
+	using namespace AudioModulation;
+	return TUniquePtr<Audio::IModulatorSettings>(new FModulationPatchSettings(*this));
+}
+
+const Audio::FModulationParameter& USoundModulationPatch::GetOutputParameter() const
+{
+	return AudioModulation::GetOrRegisterParameter(PatchSettings.OutputParameter, *this);
 }
 
 #if WITH_EDITOR
 void USoundModulationPatch::PostEditChangeProperty(FPropertyChangedEvent& InPropertyChangedEvent)
 {
-	AudioModulation::IterateModulationImpl([this](AudioModulation::FAudioModulation& OutModulation)
+	AudioModulation::IterateModulationManagers([this](AudioModulation::FAudioModulationManager& OutModulation)
 	{
 		OutModulation.UpdateModulator(*this);
 	});
@@ -38,7 +50,7 @@ void USoundModulationPatch::PostEditChangeProperty(FPropertyChangedEvent& InProp
 
 void USoundModulationPatch::PostEditChangeChainProperty(FPropertyChangedChainEvent& InPropertyChangedEvent)
 {
-	AudioModulation::IterateModulationImpl([this](AudioModulation::FAudioModulation& OutModulation)
+	AudioModulation::IterateModulationManagers([this](AudioModulation::FAudioModulationManager& OutModulation)
 	{
 		OutModulation.UpdateModulator(*this);
 	});
diff --git a/Engine/Plugins/Runtime/AudioModulation/Source/AudioModulation/Private/SoundModulationPatchProxy.cpp b/Engine/Plugins/Runtime/AudioModulation/Source/AudioModulation/Private/SoundModulationPatchProxy.cpp
--- a/Engine/Plugins/Runtime/AudioModulation/Source/AudioModulation/Private/SoundModulationPatchProxy.cpp
+++ b/Engine/Plugins/Runtime/AudioModulation/Source/AudioModulation/Private/SoundModulationPatchProxy.cpp
@@ -17,6 +17,12 @@ namespace AudioModulation
 {
 	const FPatchId InvalidPatchId = INDEX_NONE;
 
+	Audio::FModulatorTypeId FModulationPatchSettings::Register(Audio::FModulatorHandleId HandleId, IAudioModulationManager& InModulation) const
+	{
+		FAudioModulationSystem& ModSystem = static_cast<FAudioModulationManager&>(InModulation).GetSystem();
+		return ModSystem.RegisterModulator(HandleId, *this);
+	}
+
 	FModulationInputProxy::FModulationInputProxy(FModulationInputSettings&& InSettings, FAudioModulationSystem& OutModSystem)
 		: BusHandle(FBusHandle::Create(MoveTemp(InSettings.BusSettings), OutModSystem.RefProxies.Buses, OutModSystem))
 		, Transform(MoveTemp(InSettings.Transform))
@@ -38,7 +44,7 @@ namespace AudioModulation
 	void FModulationPatchProxy::Init(FModulationPatchSettings&& InSettings, FAudioModulationSystem& OutModSystem)
 	{
 		bBypass = InSettings.bBypass;
-		DefaultValue = InSettings.DefaultOutputValue;
+		DefaultValue = InSettings.OutputParameter.DefaultValue;
 
 		// Cache existing proxies to avoid releasing bus state (and potentially referenced bus state) when reinitializing
 		const TArray<FModulationInputProxy> CachedProxies = InputProxies;
@@ -49,7 +55,7 @@ namespace AudioModulation
 			InputProxies.Emplace(MoveTemp(Input), OutModSystem);
 		}
 
-		OutputProxy = FModulationOutputProxy(InSettings.DefaultOutputValue, MoveTemp(InSettings.MixFunction));
+		OutputProxy = FModulationOutputProxy(InSettings.OutputParameter.DefaultValue, MoveTemp(InSettings.OutputParameter.MixFunction));
 	}
 
 	bool FModulationPatchProxy::IsBypassed() const
diff --git a/Engine/Plugins/Runtime/AudioModulation/Source/AudioModulation/Private/SoundModulationPatchProxy.h b/Engine/Plugins/Runtime/AudioModulation/Source/AudioModulation/Private/SoundModulationPatchProxy.h
--- a/Engine/Plugins/Runtime/AudioModulation/Source/AudioModulation/Private/SoundModulationPatchProxy.h
+++ b/Engine/Plugins/Runtime/AudioModulation/Source/AudioModulation/Private/SoundModulationPatchProxy.h
@@ -23,15 +23,15 @@ namespace AudioModulation
 		FSoundModulationTransform Transform;
 		uint8 bSampleAndHold : 1;
 
-		FModulationInputSettings(const FSoundControlModulationInput& InInput, Audio::FDeviceId InDeviceId)
-			: BusSettings(InInput.GetBusChecked(), InDeviceId)
+		FModulationInputSettings(const FSoundControlModulationInput& InInput)
+			: BusSettings(InInput.GetBusChecked())
 			, Transform(InInput.Transform)
 			, bSampleAndHold(InInput.bSampleAndHold)
 		{
 		}
 
-		FModulationInputSettings(FSoundControlModulationInput&& InInput, Audio::FDeviceId InDeviceId)
-			: BusSettings(InInput.GetBusChecked(), InDeviceId)
+		FModulationInputSettings(FSoundControlModulationInput&& InInput)
+			: BusSettings(InInput.GetBusChecked())
 			, Transform(MoveTemp(InInput.Transform))
 			, bSampleAndHold(InInput.bSampleAndHold)
 		{
@@ -70,65 +70,46 @@ namespace AudioModulation
 		float DefaultValue = 1.0f;
 	};
 
-	struct FModulationPatchSettings : public TModulatorBase<FPatchId>
+	struct FModulationPatchSettings : public TModulatorBase<FPatchId>, public Audio::IModulatorSettings
 	{
-		float DefaultOutputValue = 1.0f;
-
 		TArray<FModulationInputSettings> InputSettings;
 		bool bBypass = true;
 
-		/** Function used to mix patch inputs together */
-		Audio::FModulationMixFunction MixFunction;
+		Audio::FModulationParameter OutputParameter;
 
-		FModulationPatchSettings()
-			: MixFunction(Audio::FModulationParameter::GetDefaultMixFunction())
-		{
-		}
+		FModulationPatchSettings() = default;
+		FModulationPatchSettings(const FModulationPatchSettings& InPatchSettings) = default;
 
-		FModulationPatchSettings(const FSoundControlModulationPatch& InPatch, Audio::FDeviceId InDeviceId)
-			: bBypass(InPatch.bBypass)
+		FModulationPatchSettings(const USoundModulationPatch& InPatch)
+			: TModulatorBase<FPatchId>(InPatch.GetName(), InPatch.GetUniqueID())
+			, bBypass(InPatch.PatchSettings.bBypass)
+			, OutputParameter(InPatch.GetOutputParameter())
 		{
-			if (InPatch.OutputParameter)
-			{
-				DefaultOutputValue = InPatch.OutputParameter->Settings.ValueNormalized;
-				MixFunction = InPatch.OutputParameter->GetMixFunction();
-			}
-			else
-			{
-				MixFunction = Audio::FModulationParameter::GetDefaultMixFunction();
-			}
-
-			for (const FSoundControlModulationInput& Input : InPatch.Inputs)
+			for (const FSoundControlModulationInput& Input : InPatch.PatchSettings.Inputs)
 			{
 				if (Input.GetBus())
 				{
-					InputSettings.Emplace(Input, InDeviceId);
+					InputSettings.Add(FModulationInputSettings(Input));
 				}
 			}
 		}
 
-		FModulationPatchSettings(const USoundModulationPatch& InPatch, Audio::FDeviceId InDeviceId)
-			: TModulatorBase<FPatchId>(InPatch.GetName(), InPatch.GetUniqueID())
-			, bBypass(InPatch.PatchSettings.bBypass)
+		virtual TUniquePtr<IModulatorSettings> Clone() const override
 		{
-			if (USoundModulationParameter* Parameter = InPatch.PatchSettings.OutputParameter)
-			{
-				DefaultOutputValue = Parameter->Settings.ValueNormalized;
-				MixFunction = Parameter->GetMixFunction();
-			}
-			else
-			{
-				MixFunction = Audio::FModulationParameter::GetDefaultMixFunction();
-			}
+			return TUniquePtr<IModulatorSettings>(new FModulationPatchSettings(*this));
+		}
 
-			for (const FSoundControlModulationInput& Input : InPatch.PatchSettings.Inputs)
-			{
-				if (Input.GetBus())
-				{
-					InputSettings.Emplace(Input, InDeviceId);
-				}
-			}
+		virtual Audio::FModulatorId GetModulatorId() const override
+		{
+			return static_cast<Audio::FModulatorId>(GetId());
 		}
+
+		virtual const Audio::FModulationParameter& GetOutputParameter() const override
+		{
+			return OutputParameter;
+		}
+
+		virtual Audio::FModulatorTypeId Register(Audio::FModulatorHandleId HandleId, IAudioModulationManager& InAudioModulation) const override;
 	};
 
 	class FModulationPatchProxy
diff --git a/Engine/Plugins/Runtime/AudioModulation/Source/AudioModulation/Private/SoundModulationProxy.h b/Engine/Plugins/Runtime/AudioModulation/Source/AudioModulation/Private/SoundModulationProxy.h
--- a/Engine/Plugins/Runtime/AudioModulation/Source/AudioModulation/Private/SoundModulationProxy.h
+++ b/Engine/Plugins/Runtime/AudioModulation/Source/AudioModulation/Private/SoundModulationProxy.h
@@ -12,7 +12,7 @@ namespace AudioModulation
 	class FAudioModulationSystem;
 
 	/*
-	 * Handle for all ref-counted proxy types, to be used only on the audio thread (not thread safe). 
+	 * Handle for all ref-counted proxy types, to be used only on the modulator processing thread (not thread safe).
 	 */
 	template<typename IdType, typename ProxyType, typename ProxySettings>
 	class TProxyHandle
diff --git a/Engine/Plugins/Runtime/AudioModulation/Source/AudioModulation/Private/SoundModulatorAsset.cpp b/Engine/Plugins/Runtime/AudioModulation/Source/AudioModulation/Private/SoundModulatorAsset.cpp
--- a/Engine/Plugins/Runtime/AudioModulation/Source/AudioModulation/Private/SoundModulatorAsset.cpp
+++ b/Engine/Plugins/Runtime/AudioModulation/Source/AudioModulation/Private/SoundModulatorAsset.cpp
@@ -8,35 +8,33 @@ namespace AudioModulation
 	const FText PluginAuthor = NSLOCTEXT("AudioModulation", "PluginAuthor", "Epic Games, Inc.");
 	const FText PluginNodeMissingPrompt = NSLOCTEXT("AudioModulation", "DefaultMissingNodePrompt", "The node was likely removed, renamed, or the AudioModulation plugin is not loaded.");
 
-	FSoundModulatorAsset::FSoundModulatorAsset(const Audio::IProxyDataPtr& InInitData)
+	FSoundModulatorAsset::FSoundModulatorAsset(const Audio::IProxyDataPtr& InProxyPtr)
 	{
-		if (!InInitData.IsValid())
+		if (!InProxyPtr.IsValid())
 		{
 			return;
 		}
 
-		if (!ensure(InInitData->CheckTypeCast<FSoundModulatorAssetProxy>()))
+		if (ensure(InProxyPtr->CheckTypeCast<FSoundModulatorAssetProxy>()))
 		{
-			return;
+			FSoundModulatorAssetProxy* ModulatorProxy = static_cast<FSoundModulatorAssetProxy*>(InProxyPtr.Get());
+			Proxy = MakeShared<FSoundModulatorAssetProxy, ESPMode::ThreadSafe>(*ModulatorProxy);
 		}
-
-		FSoundModulatorAssetProxy* NewProxy = static_cast<FSoundModulatorAssetProxy*>(InInitData->Clone().Release());
-		Proxy = TSharedPtr<FSoundModulatorAssetProxy, ESPMode::ThreadSafe>(NewProxy);
 	}
 
-	FSoundModulationParameterAsset::FSoundModulationParameterAsset(const Audio::IProxyDataPtr& InInitData)
+	FSoundModulationParameterAsset::FSoundModulationParameterAsset(const Audio::IProxyDataPtr& InProxyPtr)
 	{
-		if (!InInitData.IsValid())
+		if (!InProxyPtr.IsValid())
 		{
 			return;
 		}
 
-		if (!ensure(InInitData->CheckTypeCast<FSoundModulationParameterAssetProxy>()))
+		if (ensure(InProxyPtr->CheckTypeCast<FSoundModulationParameterAssetProxy>()))
 		{
-			return;
-		}
+			FSoundModulationParameterAssetProxy* ParamProxy = static_cast<FSoundModulationParameterAssetProxy*>(InProxyPtr.Get());
+			check(ParamProxy);
 
-		FSoundModulationParameterAssetProxy* NewProxy = static_cast<FSoundModulationParameterAssetProxy*>(InInitData->Clone().Release());
-		Proxy = TSharedPtr<FSoundModulationParameterAssetProxy, ESPMode::ThreadSafe>(NewProxy);
+			Proxy = MakeShared<FSoundModulationParameterAssetProxy, ESPMode::ThreadSafe>(*ParamProxy);
+		}
 	}
 } // namespace AudioModulation
diff --git a/Engine/Plugins/Runtime/AudioModulation/Source/AudioModulation/Private/SoundModulatorAssetProxy.h b/Engine/Plugins/Runtime/AudioModulation/Source/AudioModulation/Private/SoundModulatorAssetProxy.h
deleted file mode 100644
diff --git a/Engine/Plugins/Runtime/AudioModulation/Source/AudioModulation/Public/AudioModulation.h b/Engine/Plugins/Runtime/AudioModulation/Source/AudioModulation/Public/AudioModulation.h
--- a/Engine/Plugins/Runtime/AudioModulation/Source/AudioModulation/Public/AudioModulation.h
+++ b/Engine/Plugins/Runtime/AudioModulation/Source/AudioModulation/Public/AudioModulation.h
@@ -25,18 +25,40 @@ namespace AudioModulation
 	// Forward Declarations
 	class FAudioModulationSystem;
 
-	class AUDIOMODULATION_API FAudioModulation : public IAudioModulation
+	class AUDIOMODULATION_API FAudioModulationManager : public IAudioModulationManager
 	{
 	public:
-		FAudioModulation();
-		virtual ~FAudioModulation();
+		FAudioModulationManager();
+		virtual ~FAudioModulationManager();
 
-		//~ Begin IAudioModulation implementation
-		virtual Audio::FModulationParameter GetParameter(FName InParamName) override;
+		//~ Begin IAudioModulationManager implementation
 		virtual void Initialize(const FAudioPluginInitializationParams& InitializationParams) override;
-
 		virtual void OnAuditionEnd() override;
 
+#if !UE_BUILD_SHIPPING
+		virtual void SetDebugBusFilter(const FString* InFilter);
+		virtual void SetDebugMixFilter(const FString* InFilter);
+		virtual void SetDebugMatrixEnabled(bool bInIsEnabled);
+		virtual void SetDebugGeneratorsEnabled(bool bInIsEnabled);
+		virtual void SetDebugGeneratorFilter(const FString* InFilter);
+		virtual void SetDebugGeneratorTypeFilter(const FString* InFilter, bool bInIsEnabled);
+
+		virtual bool OnPostHelp(FCommonViewportClient* ViewportClient, const TCHAR* Stream) override;
+		virtual int32 OnRenderStat(FViewport* Viewport, FCanvas* Canvas, int32 X, int32 Y, const UFont& Font, const FVector* ViewLocation, const FRotator* ViewRotation) override;
+		virtual bool OnToggleStat(FCommonViewportClient* ViewportClient, const TCHAR* Stream) override;
+#endif // !UE_BUILD_SHIPPING
+
+		virtual void ProcessModulators(const double InElapsed) override;
+		virtual void UpdateModulator(const USoundModulatorBase& InModulator) override;
+
+	protected:
+		virtual void RegisterModulator(Audio::FModulatorHandleId InHandleId, Audio::FModulatorId InModulatorId) override;
+		virtual bool GetModulatorValue(const Audio::FModulatorHandle& ModulatorHandle, float& OutValue) const override;
+		virtual bool GetModulatorValueThreadSafe(const Audio::FModulatorHandle& ModulatorHandle, float& OutValue) const override;
+		virtual void UnregisterModulator(const Audio::FModulatorHandle& InHandle) override;
+	//~ End IAudioModulationManager implementation
+
+	public:
 		void ActivateBus(const USoundControlBus& InBus);
 		void ActivateBusMix(const USoundControlBusMix& InBusMix);
 		void ActivateGenerator(const USoundModulationGenerator& InGenerator);
@@ -59,50 +81,15 @@ namespace AudioModulation
 		void ClearGlobalBusMixValue(const USoundControlBus& InBus, float InFadeTime);
 		void ClearAllGlobalBusMixValues(float InFadeTime);
 
-#if !UE_BUILD_SHIPPING
-		virtual void SetDebugBusFilter(const FString* InFilter);
-		virtual void SetDebugMixFilter(const FString* InFilter);
-		virtual void SetDebugMatrixEnabled(bool bInIsEnabled);
-		virtual void SetDebugGeneratorsEnabled(bool bInIsEnabled);
-		virtual void SetDebugGeneratorFilter(const FString* InFilter);
-		virtual void SetDebugGeneratorTypeFilter(const FString* InFilter, bool bInIsEnabled);
-
-		virtual bool OnPostHelp(FCommonViewportClient* ViewportClient, const TCHAR* Stream) override;
-		virtual int32 OnRenderStat(FViewport* Viewport, FCanvas* Canvas, int32 X, int32 Y, const UFont& Font, const FVector* ViewLocation, const FRotator* ViewRotation) override;
-		virtual bool OnToggleStat(FCommonViewportClient* ViewportClient, const TCHAR* Stream) override;
-#endif // !UE_BUILD_SHIPPING
-
-		virtual void ProcessModulators(const double InElapsed) override;
-
-		virtual void UpdateModulator(const USoundModulatorBase& InModulator) override;
-		//~ End IAudioModulation implementation
-
-	protected:
-		virtual Audio::FModulatorTypeId RegisterModulator(Audio::FModulatorHandleId InHandleId, const USoundModulatorBase* InModulatorBase, Audio::FModulationParameter& OutParameter) override;
-		virtual void RegisterModulator(Audio::FModulatorHandleId InHandleId, Audio::FModulatorId InModulatorId) override;
-		virtual bool GetModulatorValue(const Audio::FModulatorHandle& ModulatorHandle, float& OutValue) const override;
-		virtual bool GetModulatorValueThreadSafe(const Audio::FModulatorHandle& ModulatorHandle, float& OutValue) const override;
-		virtual void UnregisterModulator(const Audio::FModulatorHandle& InHandle) override;
+		FAudioModulationSystem& GetSystem();
 
 	private:
 		FAudioModulationSystem* ModSystem = nullptr;
 	};
 
-	static void IterateModulationImpl(TUniqueFunction<void(FAudioModulation&)> InFunction)
-	{
-		if (FAudioDeviceManager* DeviceManager = FAudioDeviceManager::Get())
-		{
-			TArray<FAudioDevice*> Devices = DeviceManager->GetAudioDevices();
-			DeviceManager->IterateOverAllDevices([ModFunction = MoveTemp(InFunction)](Audio::FDeviceId DeviceId, FAudioDevice* AudioDevice)
-			{
-				if (AudioDevice && AudioDevice->IsModulationPluginEnabled() && AudioDevice->ModulationInterface.IsValid())
-				{
-					auto ModulationInterface = static_cast<AudioModulation::FAudioModulation*>(AudioDevice->ModulationInterface.Get());
-					ModFunction(*ModulationInterface);
-				}
-			});
-		}
-	}
+	AUDIOMODULATION_API FAudioModulationManager* GetDeviceModulationManager(Audio::FDeviceId InDeviceId);
+
+	AUDIOMODULATION_API void IterateModulationManagers(TFunctionRef<void(FAudioModulationManager&)> InFunction);
 } // namespace AudioModulation
 
 class FAudioModulationPluginFactory : public IAudioModulationFactory
diff --git a/Engine/Plugins/Runtime/AudioModulation/Source/AudioModulation/Public/AudioModulationSettings.h b/Engine/Plugins/Runtime/AudioModulation/Source/AudioModulation/Public/AudioModulationSettings.h
--- a/Engine/Plugins/Runtime/AudioModulation/Source/AudioModulation/Public/AudioModulationSettings.h
+++ b/Engine/Plugins/Runtime/AudioModulation/Source/AudioModulation/Public/AudioModulationSettings.h
@@ -1,11 +1,11 @@
 // Copyright Epic Games, Inc. All Rights Reserved.
 #pragma once
 
-#include "CoreMinimal.h"
-
 #include "Engine/DeveloperSettings.h"
+#include "Misc/Build.h"
 #include "UObject/ObjectMacros.h"
 #include "UObject/SoftObjectPath.h"
+#include "UObject/UnrealType.h"
 
 #include "AudioModulationSettings.generated.h"
 
@@ -19,4 +19,10 @@ public:
 	// Array of loaded Modulation Parameters
 	UPROPERTY(config, EditAnywhere, Category = "Parameters", meta = (AllowedClasses = "SoundModulationParameter"))
 	TArray<FSoftObjectPath> Parameters;
+
+	void RegisterParameters() const;
+
+#if WITH_EDITOR
+	virtual void PostEditChangeChainProperty(FPropertyChangedChainEvent& PropertyChangedEvent) override;
+#endif // WITH_EDITOR
 };
\ No newline at end of file
diff --git a/Engine/Plugins/Runtime/AudioModulation/Source/AudioModulation/Public/AudioModulationStatics.h b/Engine/Plugins/Runtime/AudioModulation/Source/AudioModulation/Public/AudioModulationStatics.h
--- a/Engine/Plugins/Runtime/AudioModulation/Source/AudioModulation/Public/AudioModulationStatics.h
+++ b/Engine/Plugins/Runtime/AudioModulation/Source/AudioModulation/Public/AudioModulationStatics.h
@@ -13,7 +13,7 @@
 // Forward Declarations
 namespace AudioModulation
 {
-	class FAudioModulation;
+	class FAudioModulationManager;
 	class FAudioModulationSystem;
 } // namespace AudioModulation
 
@@ -32,7 +32,7 @@ public:
 	/**
 	 * Returns modulation implementation associated with the provided world
 	 */
-	static AudioModulation::FAudioModulation* GetModulation(UWorld* World);
+	static AudioModulation::FAudioModulationManager* GetModulation(UWorld* World);
 
 	/** Manually activates a modulation bus. If called, deactivation will only occur
 	 * if bus is manually deactivated or destroyed (i.e. will not deactivate
diff --git a/Engine/Plugins/Runtime/AudioModulation/Source/AudioModulation/Public/Generators/SoundModulationEnvelopeFollower.h b/Engine/Plugins/Runtime/AudioModulation/Source/AudioModulation/Public/Generators/SoundModulationEnvelopeFollower.h
--- a/Engine/Plugins/Runtime/AudioModulation/Source/AudioModulation/Public/Generators/SoundModulationEnvelopeFollower.h
+++ b/Engine/Plugins/Runtime/AudioModulation/Source/AudioModulation/Public/Generators/SoundModulationEnvelopeFollower.h
@@ -59,5 +59,5 @@ public:
 	UPROPERTY(EditAnywhere, Category = Modulation, BlueprintReadWrite, meta = (ShowOnlyInnerProperties))
 	FEnvelopeFollowerGeneratorParams Params;
 
-	virtual AudioModulation::FGeneratorPtr CreateInstance(Audio::FDeviceId InDeviceId) const override;
+	virtual AudioModulation::FGeneratorPtr CreateInstance() const override;
 };
diff --git a/Engine/Plugins/Runtime/AudioModulation/Source/AudioModulation/Public/Generators/SoundModulationLFO.h b/Engine/Plugins/Runtime/AudioModulation/Source/AudioModulation/Public/Generators/SoundModulationLFO.h
--- a/Engine/Plugins/Runtime/AudioModulation/Source/AudioModulation/Public/Generators/SoundModulationLFO.h
+++ b/Engine/Plugins/Runtime/AudioModulation/Source/AudioModulation/Public/Generators/SoundModulationLFO.h
@@ -66,5 +66,5 @@ public:
 	UPROPERTY(EditAnywhere, Category = Modulation, BlueprintReadWrite, meta = (ShowOnlyInnerProperties))
 	FSoundModulationLFOParams Params;
 
-	virtual AudioModulation::FGeneratorPtr CreateInstance(Audio::FDeviceId InDeviceId) const override;
+	virtual AudioModulation::FGeneratorPtr CreateInstance() const override;
 };
diff --git a/Engine/Plugins/Runtime/AudioModulation/Source/AudioModulation/Public/SoundControlBus.h b/Engine/Plugins/Runtime/AudioModulation/Source/AudioModulation/Public/SoundControlBus.h
--- a/Engine/Plugins/Runtime/AudioModulation/Source/AudioModulation/Public/SoundControlBus.h
+++ b/Engine/Plugins/Runtime/AudioModulation/Source/AudioModulation/Public/SoundControlBus.h
@@ -4,7 +4,7 @@
 #include "AudioDefines.h"
 #include "AudioDeviceManager.h"
 #include "AudioDevice.h"
-#include "CoreMinimal.h"
+#include "Containers/Array.h"
 #include "IAudioModulation.h"
 #include "SoundModulationGenerator.h"
 #include "SoundModulationParameter.h"
@@ -46,6 +46,7 @@ public:
 	UPROPERTY(EditAnywhere, Category = General, BlueprintReadOnly)
 	USoundModulationParameter* Parameter;
 
+	/* UObject Implementation */
 #if WITH_EDITOR
 	virtual void PostDuplicate(EDuplicateMode::Type DuplicateMode) override;
 	virtual void PostEditChangeProperty(FPropertyChangedEvent& InPropertyChangedEvent) override;
@@ -55,18 +56,13 @@ public:
 
 	virtual void BeginDestroy() override;
 
+	/* USoundModulatorBase Implementation */
+	virtual TUniquePtr<Audio::IModulatorSettings> CreateProxySettings() const override;
 	virtual TUniquePtr<Audio::IProxyData> CreateNewProxyData(const Audio::FProxyDataInitParams& InitParams) override;
+	virtual const Audio::FModulationParameter& GetOutputParameter() const override;
 
 	const Audio::FModulationMixFunction GetMixFunction() const;
 
 	float GetDefaultNormalizedValue() const { return Parameter ? Parameter->Settings.ValueNormalized : 1.0f; }
 
-	FName GetOutputParameterName() const override
-	{
-		if (Parameter)
-		{
-			return Parameter->GetFName();
-		}
-		return Super::GetOutputParameterName();
-	}
 };
diff --git a/Engine/Plugins/Runtime/AudioModulation/Source/AudioModulation/Public/SoundModulationGenerator.h b/Engine/Plugins/Runtime/AudioModulation/Source/AudioModulation/Public/SoundModulationGenerator.h
--- a/Engine/Plugins/Runtime/AudioModulation/Source/AudioModulation/Public/SoundModulationGenerator.h
+++ b/Engine/Plugins/Runtime/AudioModulation/Source/AudioModulation/Public/SoundModulationGenerator.h
@@ -1,6 +1,7 @@
 // Copyright Epic Games, Inc. All Rights Reserved.
 #pragma once
 
+#include "AudioDefines.h"
 #include "AudioDeviceManager.h"
 #include "Containers/Queue.h"
 #include "DSP/LFO.h"
@@ -21,22 +22,6 @@ namespace AudioModulation
 	class AUDIOMODULATION_API IGenerator
 	{
 	public:
-		IGenerator() = default;
-		IGenerator(const IGenerator& InGenerator)
-			: AudioDeviceId(InGenerator.AudioDeviceId)
-		{
-		}
-
-		IGenerator(IGenerator&& InGenerator)
-			: AudioDeviceId(InGenerator.AudioDeviceId)
-		{
-		}
-
-		IGenerator(Audio::FDeviceId InDeviceId)
-			: AudioDeviceId(InDeviceId)
-		{
-		}
-
 		virtual ~IGenerator() = default;
 
 		/** Pumps commands from Audio Thread to the generator's modulation processing thread.*/
@@ -56,6 +41,9 @@ namespace AudioModulation
 		/** Returns current value of the generator. */
 		virtual float GetValue() const = 0;
 
+		/** (Optional) Initializer step where the generator is provided the associated parent AudioDevice's Id. */
+		virtual void Init(Audio::FDeviceId InDeviceId) { }
+
 		/** Returns whether or not the generator is bypassed. */
 		virtual bool IsBypassed() const = 0;
 
@@ -78,7 +66,6 @@ namespace AudioModulation
 
 	private:
 		TQueue<TUniqueFunction<void()>> CommandQueue;
-
 	};
 } // namespace AudioModulation
 
@@ -92,12 +79,15 @@ class AUDIOMODULATION_API USoundModulationGenerator : public USoundModulatorBase
 	GENERATED_BODY()
 
 public:
-	virtual AudioModulation::FGeneratorPtr CreateInstance(Audio::FDeviceId AudioDeviceId) const
+	// Create and return pointer to new instance of generator to be processed on the AudioRenderThread.
+	virtual AudioModulation::FGeneratorPtr CreateInstance() const
 	{
 		return nullptr;
 	}
 
+	/* USoundModulatorBase Implementation */
 	virtual TUniquePtr<Audio::IProxyData> CreateNewProxyData(const Audio::FProxyDataInitParams& InitParams) override;
+	virtual TUniquePtr<Audio::IModulatorSettings> CreateProxySettings() const override;
 
 #if WITH_EDITOR
 	virtual void PostEditChangeProperty(FPropertyChangedEvent& InPropertyChangedEvent) override;
diff --git a/Engine/Plugins/Runtime/AudioModulation/Source/AudioModulation/Public/SoundModulationParameter.h b/Engine/Plugins/Runtime/AudioModulation/Source/AudioModulation/Public/SoundModulationParameter.h
--- a/Engine/Plugins/Runtime/AudioModulation/Source/AudioModulation/Public/SoundModulationParameter.h
+++ b/Engine/Plugins/Runtime/AudioModulation/Source/AudioModulation/Public/SoundModulationParameter.h
@@ -10,6 +10,7 @@
 
 #include "SoundModulationParameter.generated.h"
 
+
 USTRUCT(BlueprintType)
 struct AUDIOMODULATION_API FSoundModulationParameterSettings
 {
@@ -110,6 +111,9 @@ public:
 	void RefreshNormalizedValue();
 	void RefreshUnitValue();
 #endif // WITH_EDITOR
+
+	Audio::FModulationParameter CreateParameter() const;
+
 };
 
 // Linearly scaled value between unit minimum and maximum.
@@ -254,13 +258,12 @@ namespace AudioModulation
 	{
 	public:
 		explicit FSoundModulationPluginParameterAssetProxy(USoundModulationParameter* InParameter);
+		FSoundModulationPluginParameterAssetProxy(const FSoundModulationPluginParameterAssetProxy& InProxy) = default;
 
 		virtual Audio::IProxyDataPtr Clone() const override;
 
 		virtual const Audio::FModulationParameter& GetParameter() const override;
-
-	private:
-		Audio::FModulatorHandle ModHandle;
-		Audio::FModulationParameter Parameter;
 	};
+
+	AUDIOMODULATION_API const Audio::FModulationParameter& GetOrRegisterParameter(const USoundModulationParameter* InParameter, const USoundModulatorBase& InReferencingModulator);
 } // namespace AudioModulation
\ No newline at end of file
diff --git a/Engine/Plugins/Runtime/AudioModulation/Source/AudioModulation/Public/SoundModulationPatch.h b/Engine/Plugins/Runtime/AudioModulation/Source/AudioModulation/Public/SoundModulationPatch.h
--- a/Engine/Plugins/Runtime/AudioModulation/Source/AudioModulation/Public/SoundModulationPatch.h
+++ b/Engine/Plugins/Runtime/AudioModulation/Source/AudioModulation/Public/SoundModulationPatch.h
@@ -64,17 +64,11 @@ public:
 	UPROPERTY(EditAnywhere, BlueprintReadWrite, Category = Modulation, meta = (ShowOnlyInnerProperties))
 	FSoundControlModulationPatch PatchSettings;
 
-	virtual FName GetOutputParameterName() const override
-	{
-		if (PatchSettings.OutputParameter)
-		{
-			return PatchSettings.OutputParameter->GetFName();
-		}
-
-		return Super::GetOutputParameterName();
-	}
-
+	/* USoundModulatorBase Implementation */
 	virtual TUniquePtr<Audio::IProxyData> CreateNewProxyData(const Audio::FProxyDataInitParams& InitParams) override;
+	virtual const Audio::FModulationParameter& GetOutputParameter() const override;
+
+	virtual TUniquePtr<Audio::IModulatorSettings> CreateProxySettings() const override;
 
 #if WITH_EDITOR
 	virtual void PostEditChangeProperty(FPropertyChangedEvent& InPropertyChangedEvent) override;
diff --git a/Engine/Plugins/Runtime/AudioModulation/Source/AudioModulation/Public/SoundModulatorAsset.h b/Engine/Plugins/Runtime/AudioModulation/Source/AudioModulation/Public/SoundModulatorAsset.h
--- a/Engine/Plugins/Runtime/AudioModulation/Source/AudioModulation/Public/SoundModulatorAsset.h
+++ b/Engine/Plugins/Runtime/AudioModulation/Source/AudioModulation/Public/SoundModulatorAsset.h
@@ -38,6 +38,16 @@ namespace AudioModulation
 
 		FSoundModulatorAsset(const Audio::IProxyDataPtr& InInitData);
 
+		Audio::FModulatorId GetModulatorId() const
+		{
+			if (Proxy.IsValid())
+			{
+				return Proxy->GetModulatorId();
+			}
+
+			return INDEX_NONE;
+		}
+
 		bool IsValid() const
 		{
 			return Proxy.IsValid();
diff --git a/Engine/Plugins/Runtime/Metasound/Source/MetasoundEngine/Private/MetasoundSource.cpp b/Engine/Plugins/Runtime/Metasound/Source/MetasoundEngine/Private/MetasoundSource.cpp
--- a/Engine/Plugins/Runtime/Metasound/Source/MetasoundEngine/Private/MetasoundSource.cpp
+++ b/Engine/Plugins/Runtime/Metasound/Source/MetasoundEngine/Private/MetasoundSource.cpp
@@ -3,6 +3,7 @@
 
 #include "Algo/Transform.h"
 #include "AssetRegistryModule.h"
+#include "AudioDeviceManager.h"
 #include "IAudioParameterInterfaceRegistry.h"
 #include "Internationalization/Text.h"
 #include "MetasoundAssetBase.h"
@@ -35,38 +36,72 @@
 
 #define LOCTEXT_NAMESPACE "MetaSound"
 
-namespace MetaSoundSourcePrivate
+namespace Metasound
 {
-	using FFormatOutputVertexKeyMap = TMap<EMetasoundSourceAudioFormat, TArray<Metasound::FVertexName>>;
-
-	const FFormatOutputVertexKeyMap& GetFormatOutputVertexKeys()
+	namespace SourcePrivate
 	{
-		auto CreateVertexKeyMap = []()
+		using FFormatOutputVertexKeyMap = TMap<EMetasoundSourceAudioFormat, TArray<Metasound::FVertexName>>;
+
+		// Creates a base Metasound Environment that contains basic
+		// data associated with the AudioEngine (non-source/instance specific data).
+		FMetasoundEnvironment CreateEnvironment(UWorld* World)
 		{
-			using namespace Metasound::Frontend;
+			using namespace Frontend;
+
+			FMetasoundEnvironment Environment;
+			// Add audio device ID to environment.
+			Audio::FDeviceId AudioDeviceID = INDEX_NONE;
 
-			return FFormatOutputVertexKeyMap
+			if (World)
 			{
+				FAudioDeviceHandle DeviceHandle = World->GetAudioDevice();
+				if (DeviceHandle.IsValid())
 				{
-					EMetasoundSourceAudioFormat::Mono,
-					{
-						OutputFormatMonoInterface::Outputs::MonoOut
-					}
-				},
+					AudioDeviceID = DeviceHandle.GetDeviceID();
+				}
+			}
+
+			if (INDEX_NONE == AudioDeviceID)
+			{
+				if (FAudioDeviceManager* DeviceManager = FAudioDeviceManager::Get())
+				{
+					AudioDeviceID = DeviceManager->GetMainAudioDeviceID();
+				}
+			}
+
+			Environment.SetValue<Audio::FDeviceId>(SourceInterface::Environment::DeviceID, AudioDeviceID);
+
+			return Environment;
+		}
+
+		const FFormatOutputVertexKeyMap& GetFormatOutputVertexKeys()
+		{
+			auto CreateVertexKeyMap = []()
+			{
+				using namespace Metasound::Frontend;
+
+				return FFormatOutputVertexKeyMap
 				{
-					EMetasoundSourceAudioFormat::Stereo,
 					{
-						OutputFormatStereoInterface::Outputs::LeftOut,
-						OutputFormatStereoInterface::Outputs::RightOut,
+						EMetasoundSourceAudioFormat::Mono,
+						{
+							OutputFormatMonoInterface::Outputs::MonoOut
+						}
+					},
+					{
+						EMetasoundSourceAudioFormat::Stereo,
+						{
+							OutputFormatStereoInterface::Outputs::LeftOut,
+							OutputFormatStereoInterface::Outputs::RightOut,
+						}
 					}
-				}
+				};
 			};
-		};
-		static const FFormatOutputVertexKeyMap Map = CreateVertexKeyMap();
-		return Map;
-	}
-}
-
+			static const FFormatOutputVertexKeyMap Map = CreateVertexKeyMap();
+			return Map;
+		}
+	} // namespace SourcePrivate
+} // namespace Metasound
 
 FAutoConsoleVariableRef CVarMetaSoundBlockRate(
 	TEXT("au.MetaSound.BlockRate"),
@@ -725,31 +760,9 @@ Metasound::FOperatorSettings UMetaSoundSource::GetOperatorSettings(Metasound::FS
 Metasound::FMetasoundEnvironment UMetaSoundSource::CreateEnvironment() const
 {
 	using namespace Metasound;
-	using namespace Metasound::Engine;
 	using namespace Metasound::Frontend;
 
-	FMetasoundEnvironment Environment;
-	// Add audio device ID to environment.
-	Audio::FDeviceId AudioDeviceID = INDEX_NONE;
-
-	if (UWorld* World = GetWorld())
-	{
-		FAudioDeviceHandle DeviceHandle = World->GetAudioDevice();
-		if (DeviceHandle.IsValid())
-		{
-			AudioDeviceID = DeviceHandle.GetDeviceID();
-		}
-	}
-
-	if (INDEX_NONE == AudioDeviceID)
-	{
-		if (FAudioDeviceManager* DeviceManager = FAudioDeviceManager::Get())
-		{
-			AudioDeviceID = DeviceManager->GetMainAudioDeviceID();
-		}
-	}
-
-	Environment.SetValue<Audio::FDeviceId>(SourceInterface::Environment::DeviceID, AudioDeviceID);
+	FMetasoundEnvironment Environment = SourcePrivate::CreateEnvironment(GetWorld());
 	Environment.SetValue<uint32>(SourceInterface::Environment::SoundUniqueID, GetUniqueID());
 
 	return Environment;
@@ -786,7 +799,7 @@ Metasound::FMetasoundEnvironment UMetaSoundSource::CreateEnvironment(const Audio
 
 const TArray<Metasound::FVertexName>& UMetaSoundSource::GetAudioOutputVertexKeys() const
 {
-	using namespace MetaSoundSourcePrivate;
+	using namespace Metasound::SourcePrivate;
 
 	if (const TArray<Metasound::FVertexName>* ArrayKeys = GetFormatOutputVertexKeys().Find(OutputFormat))
 	{
diff --git a/Engine/Source/Editor/AudioEditor/Private/SoundModulationDestinationLayout.cpp b/Engine/Source/Editor/AudioEditor/Private/SoundModulationDestinationLayout.cpp
--- a/Engine/Source/Editor/AudioEditor/Private/SoundModulationDestinationLayout.cpp
+++ b/Engine/Source/Editor/AudioEditor/Private/SoundModulationDestinationLayout.cpp
@@ -24,7 +24,7 @@
 
 namespace ModDestinationLayoutUtils
 {
-	IAudioModulation* GetEditorModulationInterface()
+	IAudioModulationManager* GetEditorModulationManager()
 	{
 		if (GEditor)
 		{
@@ -43,7 +43,7 @@ namespace ModDestinationLayoutUtils
 
 	bool IsModulationEnabled()
 	{
-		return GetEditorModulationInterface() != nullptr;
+		return GetEditorModulationManager() != nullptr;
 	}
 
 	FName GetParameterNameFromMetaData(const TSharedRef<IPropertyHandle>& InHandle)
@@ -80,23 +80,20 @@ namespace ModDestinationLayoutUtils
 			return false;
 		}
 
-		if (IAudioModulation* ModulationInterface = ModDestinationLayoutUtils::GetEditorModulationInterface())
+		const FName ModParamName = ModBase->GetOutputParameter().ParameterName;
+		const FName DestParamName = ModDestinationLayoutUtils::GetParameterNameFromMetaData(StructPropertyHandle);
+		if (ModParamName != FName() && DestParamName != FName() && ModParamName != DestParamName)
 		{
-			const FName ModParamName = ModBase->GetOutputParameterName();
-			const FName DestParamName = ModDestinationLayoutUtils::GetParameterNameFromMetaData(StructPropertyHandle);
-			if (ModParamName != FName() && DestParamName != FName() && ModParamName != DestParamName)
+			if (OutModParamName)
 			{
-				if (OutModParamName)
-				{
-					*OutModParamName = ModParamName;
-				}
+				*OutModParamName = ModParamName;
+			}
 
-				if (OutDestParamName)
-				{
-					*OutDestParamName = DestParamName;
-				}
-				return true;
+			if (OutDestParamName)
+			{
+				*OutDestParamName = DestParamName;
 			}
+			return true;
 		}
 
 		return false;
@@ -126,20 +123,16 @@ namespace ModDestinationLayoutUtils
 		OutParamName = ModDestinationLayoutUtils::GetParameterNameFromMetaData(StructPropertyHandle);
 		if (OutParamName != FName())
 		{
-			// If parameter was provided, it overrides ClampMin/Max.  User data however overrides UIMin/Max if its
-			// in clamp range.
-			if (IAudioModulation* ModulationInterface = ModDestinationLayoutUtils::GetEditorModulationInterface())
+			// If parameter was provided, it overrides ClampMin/Max.  User data however overrides UIMin/Max if its in clamp range.
+			const Audio::FModulationParameter& Parameter = Audio::GetModulationParameter(OutParamName);
+			UIMinValue = Parameter.MinValue;
+			UIMaxValue = Parameter.MaxValue;
+			ClampMinValue = UIMinValue;
+			ClampMaxValue = UIMaxValue;
+			OutUnitDisplayText = Parameter.UnitDisplayName;
+			if (bClampValuesSet)
 			{
-				Audio::FModulationParameter Parameter = ModulationInterface->GetParameter(OutParamName);
-				UIMinValue = Parameter.MinValue;
-				UIMaxValue = Parameter.MaxValue;
-				ClampMinValue = UIMinValue;
-				ClampMaxValue = UIMaxValue;
-				OutUnitDisplayText = Parameter.UnitDisplayName;
-				if (bClampValuesSet)
-				{
-					UE_LOG(LogAudioEditor, Warning, TEXT("ClampMin/Max overridden by AudioModulation plugin asset with ParamName '%s'."), *OutParamName.ToString());
-				}
+				UE_LOG(LogAudioEditor, Warning, TEXT("ClampMin/Max overridden by AudioModulation plugin asset with ParamName '%s'."), *OutParamName.ToString());
 			}
 
 			if (StructPropertyHandle->HasMetaData("UIMin"))
@@ -236,24 +229,15 @@ namespace ModDestinationLayoutUtils
 						float CurrentValue = 0.0f;
 						ValueHandle->GetValue(CurrentValue);
 
-						if (IAudioModulation* ModulationInterface = ModDestinationLayoutUtils::GetEditorModulationInterface())
-						{
-							Audio::FModulationParameter Parameter = ModulationInterface->GetParameter(ParamName);
-								
-							return Parameter.DefaultValue == CurrentValue
-								? EVisibility::Hidden
-								: EVisibility::Visible;
-						}
-
-						return EVisibility::Hidden;
+						const Audio::FModulationParameter& Parameter = Audio::GetModulationParameter(ParamName);
+						return Parameter.DefaultValue == CurrentValue
+							? EVisibility::Hidden
+							: EVisibility::Visible;
 					}))
 					.OnClicked(FOnClicked::CreateLambda([ParamName, ValueHandle]()
 					{
-						if (IAudioModulation* ModulationInterface = ModDestinationLayoutUtils::GetEditorModulationInterface())
-						{
-							Audio::FModulationParameter Parameter = ModulationInterface->GetParameter(ParamName);
-							ValueHandle->SetValue(Parameter.DefaultValue);
-						}
+						Audio::FModulationParameter Parameter = Audio::GetModulationParameter(ParamName);
+						ValueHandle->SetValue(Parameter.DefaultValue);
 
 						return FReply::Handled();
 					}))
@@ -271,14 +255,8 @@ namespace ModDestinationLayoutUtils
 			EnablementHandle->GetValue(bEnabled);
 			if (bEnabled)
 			{
-				Audio::FModulationParameter Parameter;
-
-				if (IAudioModulation* ModulationInterface = ModDestinationLayoutUtils::GetEditorModulationInterface())
-				{
-					const FName ParamName = ModDestinationLayoutUtils::GetParameterNameFromMetaData(StructPropertyHandle);
-					Parameter = ModulationInterface->GetParameter(ParamName);
-				}
-
+				const FName ParamName = ModDestinationLayoutUtils::GetParameterNameFromMetaData(StructPropertyHandle);
+				const Audio::FModulationParameter& Parameter = Audio::GetModulationParameter(ParamName);
 				ValueHandle->SetValue(Parameter.DefaultValue);
 			}
 			else
diff --git a/Engine/Source/Runtime/AudioExtensions/Private/IAudioModulation.cpp b/Engine/Source/Runtime/AudioExtensions/Private/IAudioModulation.cpp
--- a/Engine/Source/Runtime/AudioExtensions/Private/IAudioModulation.cpp
+++ b/Engine/Source/Runtime/AudioExtensions/Private/IAudioModulation.cpp
@@ -1,11 +1,67 @@
 // Copyright Epic Games, Inc. All Rights Reserved.
 #include "IAudioModulation.h"
 
+#include "Containers/Map.h"
 #include "UObject/Object.h"
 
 
 namespace Audio
 {
+	namespace ModulationInterfacePrivate
+	{
+		class FModulationParameterRegistry
+		{
+			TMap<FName, FModulationParameter> Values;
+
+			mutable FCriticalSection ThreadSafeValueAccessor;
+
+		public:
+			bool IsRegistered(FName InName) const
+			{
+				FScopeLock Lock(&ThreadSafeValueAccessor);
+				return Values.Contains(InName);
+			}
+
+			void Register(FName InName, FModulationParameter&& InParameter)
+			{
+				FScopeLock Lock(&ThreadSafeValueAccessor);
+				if (FModulationParameter* Value = Values.Find(InName))
+				{
+					*Value = MoveTemp(InParameter);
+				}
+				else
+				{
+					Values.Add(InName, MoveTemp(InParameter));
+				}
+			}
+
+			bool Unregister(FName InName)
+			{
+				FScopeLock Lock(&ThreadSafeValueAccessor);
+				return Values.Remove(InName) > 0;
+			}
+
+			void UnregisterAll()
+			{
+				FScopeLock Lock(&ThreadSafeValueAccessor);
+				Values.Reset();
+			}
+
+			const FModulationParameter& Get(FName InName) const
+			{
+				FScopeLock Lock(&ThreadSafeValueAccessor);
+
+				static const FModulationParameter DefaultParameter { };
+				if (const FModulationParameter* Param = Values.Find(InName))
+				{
+					return *Param;
+				}
+
+				return DefaultParameter;
+			}
+		} ParameterRegistry;
+	}
+
 	FModulatorHandleId CreateModulatorHandleId()
 	{
 		static FModulatorHandleId NextHandleId = INDEX_NONE;
@@ -115,16 +171,51 @@ namespace Audio
 		return ConversionFunction;
 	};
 
-	FModulatorHandle::FModulatorHandle(IAudioModulation& InModulation, const USoundModulatorBase* InModulatorBase, FName InParameterName)
+	void RegisterModulationParameter(FName InName, FModulationParameter&& InParameter)
 	{
-		HandleId = CreateModulatorHandleId();
-		Modulation = InModulation.AsShared();
-		Parameter.ParameterName = InParameterName;
-		ModulatorTypeId = InModulation.RegisterModulator(HandleId, InModulatorBase, Parameter);
+		using namespace ModulationInterfacePrivate;
+		ParameterRegistry.Register(InName, MoveTemp(InParameter));
+	}
 
+	bool UnregisterModulationParameter(FName InName)
+	{
+		using namespace ModulationInterfacePrivate;
+		return ParameterRegistry.Unregister(InName);
+	}
+
+	void UnregisterAllModulationParameters()
+	{
+		using namespace ModulationInterfacePrivate;
+		ParameterRegistry.UnregisterAll();
+	}
+
+	bool IsModulationParameterRegistered(FName InName)
+	{
+		using namespace ModulationInterfacePrivate;
+		return ParameterRegistry.IsRegistered(InName);
+	}
+
+	const FModulationParameter& GetModulationParameter(FName InName)
+	{
+		using namespace ModulationInterfacePrivate;
+		return ParameterRegistry.Get(InName);
+	}
+
+	FModulatorHandle::FModulatorHandle(Audio::FModulationParameter&& InParameter)
+		: Parameter(InParameter)
+		, HandleId(CreateModulatorHandleId())
+	{
+	}
+
+	FModulatorHandle::FModulatorHandle(IAudioModulationManager& InModulation, const Audio::IModulatorSettings& InModulatorSettings, Audio::FModulationParameter&& InParameter)
+		: Parameter(MoveTemp(InParameter))
+		, HandleId(CreateModulatorHandleId())
+		, Modulation(InModulation.AsShared())
+	{
+		ModulatorTypeId = InModulatorSettings.Register(HandleId, InModulation);
 		if (ModulatorTypeId != INDEX_NONE)
 		{
-			ModulatorId = static_cast<Audio::FModulatorId>(InModulatorBase->GetUniqueID());
+			ModulatorId = InModulatorSettings.GetModulatorId();
 		}
 	}
 
@@ -132,7 +223,7 @@ namespace Audio
 	{
 		HandleId = CreateModulatorHandleId();
 
-		if (TSharedPtr<IAudioModulation> ModPtr = InOther.Modulation.Pin())
+		if (TSharedPtr<IAudioModulationManager> ModPtr = InOther.Modulation.Pin())
 		{
 			ModPtr->RegisterModulator(HandleId, InOther.ModulatorId);
 			Parameter = InOther.Parameter;
@@ -162,7 +253,7 @@ namespace Audio
 
 	FModulatorHandle::~FModulatorHandle()
 	{
-		if (TSharedPtr<IAudioModulation> ModPtr = Modulation.Pin())
+		if (TSharedPtr<IAudioModulationManager> ModPtr = Modulation.Pin())
 		{
 			ModPtr->UnregisterModulator(*this);
 		}
@@ -172,7 +263,7 @@ namespace Audio
 	{
 		Parameter = InOther.Parameter;
 
-		if (TSharedPtr<IAudioModulation> ModPtr = InOther.Modulation.Pin())
+		if (TSharedPtr<IAudioModulationManager> ModPtr = InOther.Modulation.Pin())
 		{
 			HandleId = CreateModulatorHandleId();
 			ModulatorId = InOther.ModulatorId;
@@ -199,7 +290,7 @@ namespace Audio
 	{
 		if (HandleId != INDEX_NONE)
 		{
-			if (TSharedPtr<IAudioModulation> ModPtr = Modulation.Pin())
+			if (TSharedPtr<IAudioModulationManager> ModPtr = Modulation.Pin())
 			{
 				ModPtr->UnregisterModulator(*this);
 			}
@@ -250,7 +341,7 @@ namespace Audio
 
 		OutValue = 1.0f;
 
-		if (TSharedPtr<IAudioModulation> ModPtr = Modulation.Pin())
+		if (TSharedPtr<IAudioModulationManager> ModPtr = Modulation.Pin())
 		{
 			return ModPtr->GetModulatorValue(*this, OutValue);
 		}
@@ -263,7 +354,7 @@ namespace Audio
 		check(IsValid());
 
 		OutValue = 1.0f;
-		if (TSharedPtr<IAudioModulation> ModPtr = Modulation.Pin())
+		if (TSharedPtr<IAudioModulationManager> ModPtr = Modulation.Pin())
 		{
 			return ModPtr->GetModulatorValueThreadSafe(*this, OutValue);
 		}
@@ -277,9 +368,20 @@ namespace Audio
 	}
 } // namespace Audio
 
+const Audio::FModulationParameter& USoundModulatorBase::GetOutputParameter() const
+{
+	return Audio::GetModulationParameter({ });
+}
+
 TUniquePtr<Audio::IProxyData> USoundModulatorBase::CreateNewProxyData(const Audio::FProxyDataInitParams& InitParams)
 {
 	// This should never be hit as all instances of modulators should implement their own version of the proxy data interface.
 	checkNoEntry();
 	return TUniquePtr<Audio::IProxyData>();
 }
+
+TUniquePtr<Audio::IModulatorSettings> USoundModulatorBase::CreateProxySettings() const
+{
+	checkNoEntry();
+	return TUniquePtr<Audio::IModulatorSettings>();
+}
diff --git a/Engine/Source/Runtime/AudioExtensions/Public/IAudioExtensionPlugin.h b/Engine/Source/Runtime/AudioExtensions/Public/IAudioExtensionPlugin.h
--- a/Engine/Source/Runtime/AudioExtensions/Public/IAudioExtensionPlugin.h
+++ b/Engine/Source/Runtime/AudioExtensions/Public/IAudioExtensionPlugin.h
@@ -20,7 +20,7 @@ class FSoundEffectBase;
 class FSoundEffectSource;
 class FSoundEffectSubmix;
 struct FWaveInstance;
-class IAudioModulation;
+class IAudioModulationManager;
 class IAudioOcclusion;
 class IAudioPluginListener;
 class IAudioReverb;
@@ -53,7 +53,7 @@ using TSoundEffectSourcePtr   = TSharedPtr<FSoundEffectSource, ESPMode::ThreadSa
 using TSoundEffectSubmixPtr   = TSharedPtr<FSoundEffectSubmix, ESPMode::ThreadSafe>;
 using TAudioSpatializationPtr = TSharedPtr<IAudioSpatialization, ESPMode::ThreadSafe>;
 using TAudioSourceDataOverridePtr = TSharedPtr<IAudioSourceDataOverride, ESPMode::ThreadSafe>;
-using TAudioModulationPtr     = TSharedPtr<IAudioModulation, ESPMode::ThreadSafe>;
+using TAudioModulationPtr     = TSharedPtr<IAudioModulationManager, ESPMode::ThreadSafe>;
 using TAudioOcclusionPtr      = TSharedPtr<IAudioOcclusion, ESPMode::ThreadSafe>;
 using TAudioReverbPtr         = TSharedPtr<IAudioReverb, ESPMode::ThreadSafe>;
 using TAudioPluginListenerPtr = TSharedPtr<IAudioPluginListener, ESPMode::ThreadSafe>;
diff --git a/Engine/Source/Runtime/AudioExtensions/Public/IAudioModulation.h b/Engine/Source/Runtime/AudioExtensions/Public/IAudioModulation.h
--- a/Engine/Source/Runtime/AudioExtensions/Public/IAudioModulation.h
+++ b/Engine/Source/Runtime/AudioExtensions/Public/IAudioModulation.h
@@ -5,7 +5,6 @@
 #include "DSP/BufferVectorOperations.h"
 #include "IAudioExtensionPlugin.h"
 #include "IAudioProxyInitializer.h"
-#include "Internationalization/Text.h"
 #include "Templates/SharedPointer.h"
 #include "UObject/NameTypes.h"
 
@@ -13,7 +12,7 @@
 
 
 // Forward Declarations
-class IAudioModulation;
+class IAudioModulationManager;
 class ISoundModulatable;
 class USoundModulatorBase;
 class UObject;
@@ -77,11 +76,37 @@ namespace Audio
 		static const FModulationNormalizedConversionFunction& GetDefaultNormalizedConversionFunction();
 	};
 
-	/** Handle to a modulator which interacts with the modulation API to manage lifetime of internal objects */
+	AUDIOEXTENSIONS_API bool IsModulationParameterRegistered(FName InName);
+	AUDIOEXTENSIONS_API void RegisterModulationParameter(FName InName, FModulationParameter&& InParameter);
+	AUDIOEXTENSIONS_API bool UnregisterModulationParameter(FName InName);
+	AUDIOEXTENSIONS_API void UnregisterAllModulationParameters();
+	AUDIOEXTENSIONS_API const FModulationParameter& GetModulationParameter(FName InName);
+
+	/** Interface for cached off Modulator UObject data used as default settings to
+	  * be converted to instanced proxy data per AudioDevice on the AudioRenderThread.
+	  * If proxy is already active, implementation is expected to ignore register call
+	  * and return existing modulator proxy's type Id & set parameter accordingly.
+	  */
+	class AUDIOEXTENSIONS_API IModulatorSettings
+	{
+	public:
+		virtual ~IModulatorSettings() = default;
+		virtual TUniquePtr<IModulatorSettings> Clone() const = 0;
+		virtual FModulatorId GetModulatorId() const = 0;
+		virtual const Audio::FModulationParameter& GetOutputParameter() const = 0;
+		virtual Audio::FModulatorTypeId Register(
+			Audio::FModulatorHandleId HandleId,
+			IAudioModulationManager& InModulation) const = 0;
+	};
+
+	/** Handle to a modulator which interacts with the modulation API to manage lifetime
+	  * of modulator proxy objects internal to modulation plugin implementation.
+	  */
 	struct AUDIOEXTENSIONS_API FModulatorHandle
 	{
 		FModulatorHandle() = default;
-		FModulatorHandle(IAudioModulation& InModulation, const USoundModulatorBase* InModulatorBase, FName InParameterName);
+		FModulatorHandle(Audio::FModulationParameter&& InParameter);
+		FModulatorHandle(IAudioModulationManager& InModulation, const Audio::IModulatorSettings& InModulatorSettings, Audio::FModulationParameter&& InParameter);
 		FModulatorHandle(const FModulatorHandle& InOther);
 		FModulatorHandle(FModulatorHandle&& InOther);
 
@@ -103,18 +128,15 @@ namespace Audio
 		FModulatorHandleId HandleId = INDEX_NONE;
 		FModulatorTypeId ModulatorTypeId = INDEX_NONE;
 		FModulatorId ModulatorId = INDEX_NONE;
-		TWeakPtr<IAudioModulation> Modulation;
+		TWeakPtr<IAudioModulationManager> Modulation;
 	};
 } // namespace Audio
 
-class AUDIOEXTENSIONS_API IAudioModulation : public TSharedFromThis<IAudioModulation>
+class AUDIOEXTENSIONS_API IAudioModulationManager : public TSharedFromThis<IAudioModulationManager>
 {
 public:
 	/** Virtual destructor */
-	virtual ~IAudioModulation() { }
-
-	/** Returns parameter info for the given parameter name */
-	virtual Audio::FModulationParameter GetParameter(FName InParamName) { return Audio::FModulationParameter(); }
+	virtual ~IAudioModulationManager() { }
 
 	/** Initialize the modulation plugin with the same rate and number of sources */
 	virtual void Initialize(const FAudioPluginInitializationParams& InitializationParams) { }
@@ -142,7 +164,6 @@ public:
 	virtual void UpdateModulator(const USoundModulatorBase& InModulator) { }
 
 protected:
-	virtual Audio::FModulatorTypeId RegisterModulator(uint32 InHandleId, const USoundModulatorBase* InModulatorBase, Audio::FModulationParameter& OutParameter) { return INDEX_NONE; }
 	virtual void RegisterModulator(uint32 InHandleId, Audio::FModulatorId InModulatorId) { }
 
 	// Get the modulator value from the AudioRender Thread
@@ -165,14 +186,12 @@ class AUDIOEXTENSIONS_API USoundModulatorBase : public UObject, public IAudioPro
 	GENERATED_BODY()
 
 public:
-	virtual FName GetOutputParameterName() const
-	{
-		return FName();
-	}
+	virtual const Audio::FModulationParameter& GetOutputParameter() const;
 
 	virtual TUniquePtr<Audio::IProxyData> CreateNewProxyData(const Audio::FProxyDataInitParams& InitParams) override;
-};
 
+	virtual TUniquePtr<Audio::IModulatorSettings> CreateProxySettings() const;
+};
 
 /** Proxy to modulator, allowing for modulator to be referenced by the Audio Render Thread independently
   * from the implementing modulation plugin (ex. for MetaSound implementation).
@@ -182,21 +201,40 @@ class AUDIOEXTENSIONS_API FSoundModulatorAssetProxy : public Audio::TProxyData<F
 public:
 	IMPL_AUDIOPROXY_CLASS(FSoundModulatorAssetProxy);
 
-	virtual Audio::IProxyDataPtr Clone() const override
+	FSoundModulatorAssetProxy(const FSoundModulatorAssetProxy& InAssetProxy)
+		: Parameter(InAssetProxy.Parameter)
+		, ModulatorSettings(InAssetProxy.ModulatorSettings.IsValid() ? InAssetProxy.ModulatorSettings->Clone() : nullptr)
 	{
-		return Audio::IProxyDataPtr();
 	}
 
-	virtual float GetValue() const
+	FSoundModulatorAssetProxy(const USoundModulatorBase& InModulatorBase)
+		: Parameter(InModulatorBase.GetOutputParameter())
+		, ModulatorSettings(InModulatorBase.CreateProxySettings())
 	{
-		return 1.0f;
 	}
 
-	virtual const Audio::FModulationParameter& GetParameter() const
+	virtual Audio::IProxyDataPtr Clone() const
 	{
-		static const Audio::FModulationParameter DefaultParam;
-		return DefaultParam;
+		return MakeUnique<FSoundModulatorAssetProxy>(*this);
+	}
+
+	virtual Audio::FModulatorHandle CreateModulatorHandle(IAudioModulationManager& InModulation) const
+	{
+		check(ModulatorSettings.IsValid());
+
+		Audio::FModulationParameter HandleParameter = Parameter;
+		return Audio::FModulatorHandle(InModulation, *ModulatorSettings.Get(), MoveTemp(HandleParameter));
 	}
+
+	virtual Audio::FModulatorId GetModulatorId() const
+	{
+		check(ModulatorSettings.IsValid())
+		return ModulatorSettings->GetModulatorId();
+	}
+
+protected:
+	Audio::FModulationParameter Parameter;
+	TUniquePtr<Audio::IModulatorSettings> ModulatorSettings;
 };
 using FSoundModulatorAssetProxyPtr = TSharedPtr<FSoundModulatorAssetProxy, ESPMode::ThreadSafe>;
 
@@ -218,10 +256,14 @@ public:
 		static const Audio::FModulationParameter DefaultParam;
 		return DefaultParam;
 	}
+
+	protected:
+	Audio::FModulatorHandle ModHandle;
+	Audio::FModulationParameter Parameter;
 };
 using FSoundModulationParameterAssetProxyPtr = TSharedPtr<FSoundModulationParameterAssetProxy, ESPMode::ThreadSafe>;
 
-/** Interface to sound that is modulateable, allowing for certain specific
+/** Interface to sound that is modulatable, allowing for certain specific
   * behaviors to be controlled on the sound level by the modulation system.
   */
 class AUDIOEXTENSIONS_API ISoundModulatable
diff --git a/Engine/Source/Runtime/AudioMixer/Private/AudioMixerSourceManager.cpp b/Engine/Source/Runtime/AudioMixer/Private/AudioMixerSourceManager.cpp
--- a/Engine/Source/Runtime/AudioMixer/Private/AudioMixerSourceManager.cpp
+++ b/Engine/Source/Runtime/AudioMixer/Private/AudioMixerSourceManager.cpp
@@ -751,6 +751,7 @@ namespace Audio
 		PitchMod.UpdateModulator(InitParams.ModulationSettings.PitchModulationDestination.Modulator);
 
 		FModulationDestination HighpassMod;
+		PitchMod.Init(MixerDevice->DeviceID, FName("Pitch"), false /* bInIsBuffered */);
 		HighpassMod.Init(MixerDevice->DeviceID, FName("HPFCutoffFrequency"), false /* bInIsBuffered */);
 		HighpassMod.UpdateModulator(InitParams.ModulationSettings.HighpassModulationDestination.Modulator);
 
@@ -1492,13 +1493,13 @@ namespace Audio
 		AUDIO_MIXER_CHECK_GAME_THREAD(MixerDevice);
 
 		AudioMixerThreadCommand([this, SourceId, InHPFFrequency]()
-			{
-				AUDIO_MIXER_CHECK_AUDIO_PLAT_THREAD(MixerDevice);
+		{
+			AUDIO_MIXER_CHECK_AUDIO_PLAT_THREAD(MixerDevice);
 
-				FSourceInfo& SourceInfo = SourceInfos[SourceId];
-				SourceInfo.HighpassModulationBase = InHPFFrequency;
-				SourceInfo.bModFiltersUpdated = true;
-			});
+			FSourceInfo& SourceInfo = SourceInfos[SourceId];
+			SourceInfo.HighpassModulationBase = InHPFFrequency;
+			SourceInfo.bModFiltersUpdated = true;
+		});
 	}
 
 	void FMixerSourceManager::SetModVolume(const int32 SourceId, const float InModVolume)
@@ -1523,12 +1524,12 @@ namespace Audio
 		AUDIO_MIXER_CHECK_GAME_THREAD(MixerDevice);
 
 		AudioMixerThreadCommand([this, SourceId, InModPitch]()
-			{
-				AUDIO_MIXER_CHECK_AUDIO_PLAT_THREAD(MixerDevice);
+		{
+			AUDIO_MIXER_CHECK_AUDIO_PLAT_THREAD(MixerDevice);
 
-				FSourceInfo& SourceInfo = SourceInfos[SourceId];
-				SourceInfo.PitchModulationBase = InModPitch;
-			});
+			FSourceInfo& SourceInfo = SourceInfos[SourceId];
+			SourceInfo.PitchModulationBase = InModPitch;
+		});
 	}
 
 	void FMixerSourceManager::SetSubmixSendInfo(const int32 SourceId, const FMixerSourceSubmixSend& InSubmixSend)
diff --git a/Engine/Source/Runtime/AudioMixerCore/Public/AudioDefines.h b/Engine/Source/Runtime/AudioMixerCore/Public/AudioDefines.h
--- a/Engine/Source/Runtime/AudioMixerCore/Public/AudioDefines.h
+++ b/Engine/Source/Runtime/AudioMixerCore/Public/AudioDefines.h
@@ -1,20 +1,15 @@
 // Copyright Epic Games, Inc. All Rights Reserved.
-
-/*=============================================================================
-	AudioDefines.h: Defines for audio system
-=============================================================================*/
-
 #pragma once
 
-#include "CoreMinimal.h"
+#include "HAL/Platform.h"
 
 
-/** 
+/**
  * Maximum number of channels that can be set using the ini setting
  */
 #define MAX_AUDIOCHANNELS				64
 
-/** 
+/**
  * Length of sound in seconds to be considered as looping forever
  */
 #define INDEFINITELY_LOOPING_DURATION	10000.0f
@@ -52,4 +47,13 @@
 /**
  * Debugger is Available on non-shipping builds
  */
-#define ENABLE_AUDIO_DEBUG !UE_BUILD_SHIPPING
\ No newline at end of file
+#define ENABLE_AUDIO_DEBUG !UE_BUILD_SHIPPING
+
+/** Common Audio namespace Type Definitions/Identifiers */
+namespace Audio
+{
+	/**
+	 * Typed identifier for Audio Device Id
+	 */
+	using FDeviceId = uint32;
+}
\ No newline at end of file
diff --git a/Engine/Source/Runtime/Engine/Classes/Sound/SoundModulationDestination.h b/Engine/Source/Runtime/Engine/Classes/Sound/SoundModulationDestination.h
--- a/Engine/Source/Runtime/Engine/Classes/Sound/SoundModulationDestination.h
+++ b/Engine/Source/Runtime/Engine/Classes/Sound/SoundModulationDestination.h
@@ -140,7 +140,7 @@ namespace Audio
 			void UpdateModulator(const USoundModulatorBase* InModulator);
 
 		private:
-			void ClearHandle();
+			void ResetHandle();
 			void SetHandle(const FModulatorHandle& InHandle);
 			void SetHandle(FModulatorHandle&& InHandle);
 
diff --git a/Engine/Source/Runtime/Engine/Private/AudioDebug.cpp b/Engine/Source/Runtime/Engine/Private/AudioDebug.cpp
--- a/Engine/Source/Runtime/Engine/Private/AudioDebug.cpp
+++ b/Engine/Source/Runtime/Engine/Private/AudioDebug.cpp
@@ -15,6 +15,7 @@
 #include "Engine/World.h"
 #include "GameFramework/GameUserSettings.h"
 #include "HAL/IConsoleManager.h"
+#include "IAudioModulation.h"
 #include "Misc/CommandLine.h"
 #include "Misc/ConfigCacheIni.h"
 #include "Misc/Parse.h"
@@ -1691,7 +1692,7 @@ namespace Audio
 		Y += FontHeight;
 
 		bool bDisplayedSoundModulationInfo = false;
-		if (IAudioModulation* Modulation = AudioDevice->ModulationInterface.Get())
+		if (IAudioModulationManager* Modulation = AudioDevice->ModulationInterface.Get())
 		{
 			const int32 YInit = Y;
 			Y = Modulation->OnRenderStat(Viewport, Canvas, X, Y, *GetStatsFont(), ViewLocation, ViewRotation);
@@ -2317,7 +2318,7 @@ namespace Audio
 
 		if (AudioDevice->IsModulationPluginEnabled())
 		{
-			if (IAudioModulation* Modulation = AudioDevice->ModulationInterface.Get())
+			if (IAudioModulationManager* Modulation = AudioDevice->ModulationInterface.Get())
 			{
 				if (!Modulation->OnToggleStat(ViewportClient, Stream))
 				{
@@ -2341,7 +2342,7 @@ namespace Audio
 		{
 			if (AudioDevice->IsModulationPluginEnabled())
 			{
-				if (IAudioModulation* Modulation = AudioDevice->ModulationInterface.Get())
+				if (IAudioModulationManager* Modulation = AudioDevice->ModulationInterface.Get())
 				{
 					if (!Modulation->OnPostHelp(ViewportClient, Stream))
 					{
diff --git a/Engine/Source/Runtime/Engine/Private/SoundModulationDestination.cpp b/Engine/Source/Runtime/Engine/Private/SoundModulationDestination.cpp
--- a/Engine/Source/Runtime/Engine/Private/SoundModulationDestination.cpp
+++ b/Engine/Source/Runtime/Engine/Private/SoundModulationDestination.cpp
@@ -72,6 +72,7 @@ namespace Audio
 		ValueTarget = MoveTemp(InModulationDestination.ValueTarget);
 		bIsBuffered = MoveTemp(InModulationDestination.bIsBuffered);
 		bValueNormalized = MoveTemp(InModulationDestination.bValueNormalized);
+		bHasProcessed = MoveTemp(InModulationDestination.bHasProcessed);
 		OutputBuffer = MoveTemp(InModulationDestination.OutputBuffer);
 		{
 			FScopeLock OtherLock(&InModulationDestination.HandleCritSection);
@@ -79,15 +80,18 @@ namespace Audio
 			Handle = MoveTemp(InModulationDestination.Handle);
 		}
 
+
 		ParameterName = MoveTemp(InModulationDestination.ParameterName);
 
 		return *this;
 	}
 
-	void FModulationDestination::ClearHandle()
+	void FModulationDestination::ResetHandle()
 	{
+		Audio::FModulationParameter Parameter = Audio::GetModulationParameter(ParameterName);
+
 		FScopeLock Lock(&HandleCritSection);
-		Handle = FModulatorHandle();
+		Handle = FModulatorHandle(MoveTemp(Parameter));
 	}
 
 	void FModulationDestination::Init(FDeviceId InDeviceId, bool bInIsBuffered, bool bInValueNormalized)
@@ -104,7 +108,7 @@ namespace Audio
 		OutputBuffer.Reset();
 		ParameterName = InParameterName;
 
-		ClearHandle();
+		ResetHandle();
 	}
 
 	bool FModulationDestination::IsActive()
@@ -210,22 +214,28 @@ namespace Audio
 
 	void FModulationDestination::UpdateModulator(const USoundModulatorBase* InModulator)
 	{
-		const TWeakObjectPtr<const USoundModulatorBase> ModPtr(InModulator);
-		auto UpdateHandleLambda = [this, ModPtr]()
+		if (!InModulator)
+		{
+			return;
+		}
+
+		Audio::FModulationParameter OutputParam = InModulator->GetOutputParameter();
+		auto UpdateHandleLambda = [this, ModSettings = InModulator->CreateProxySettings(), Parameter = MoveTemp(OutputParam)]() mutable
 		{
 			if (FAudioDevice* AudioDevice = FAudioDeviceManager::Get()->GetAudioDeviceRaw(DeviceId))
 			{
 				if (AudioDevice->IsModulationPluginEnabled() && AudioDevice->ModulationInterface.IsValid())
 				{
-					if (IAudioModulation* Modulation = AudioDevice->ModulationInterface.Get())
+					if (IAudioModulationManager* Modulation = AudioDevice->ModulationInterface.Get())
 					{
-						SetHandle(FModulatorHandle(*Modulation, ModPtr.Get(), ParameterName));
+						check(ModSettings.IsValid());
+						SetHandle(FModulatorHandle(*Modulation, *ModSettings.Get(), MoveTemp(Parameter)));
 					}
 					return;
 				}
 			}
 
-			ClearHandle();
+			ResetHandle();
 		};
 
 		FAudioThread::RunCommandOnAudioThread(MoveTemp(UpdateHandleLambda));
diff --git a/Engine/Source/Runtime/Engine/Public/Audio.h b/Engine/Source/Runtime/Engine/Public/Audio.h
--- a/Engine/Source/Runtime/Engine/Public/Audio.h
+++ b/Engine/Source/Runtime/Engine/Public/Audio.h
@@ -93,14 +93,6 @@ struct FWaveInstance;
 struct FReverbSettings;
 struct FSampleLoop;
 
-namespace Audio
-{
-	/**
-	 * Typed identifier for Audio Device Id
-	 */
-	using FDeviceId = uint32;
-}
-
 enum ELoopingMode
 {
 	/** One shot sound */
