Fix static analysis warning
#rnx

diff --git a/Engine/Plugins/Runtime/AudioModulation/Source/AudioModulation/Private/AudioModulationSystem.h b/Engine/Plugins/Runtime/AudioModulation/Source/AudioModulation/Private/AudioModulationSystem.h
--- a/Engine/Plugins/Runtime/AudioModulation/Source/AudioModulation/Private/AudioModulationSystem.h
+++ b/Engine/Plugins/Runtime/AudioModulation/Source/AudioModulation/Private/AudioModulationSystem.h
@@ -21,6 +21,11 @@
 #include "Templates/Atomic.h"
 #include "Templates/Function.h"
 
+#if WITH_AUDIOMODULATION
+#if !UE_BUILD_SHIPPING
+#include "AudioModulationDebugger.h"
+#endif // !UE_BUILD_SHIPPING
+#endif // WITH_AUDIOMODULATION
 
 namespace AudioModulation
 {
@@ -53,12 +58,8 @@ namespace AudioModulation
 	};
 } // namespace AudioModulation
 
-#if WITH_AUDIOMODULATION
-
-#if !UE_BUILD_SHIPPING
-#include "AudioModulationDebugger.h"
-#endif // !UE_BUILD_SHIPPING
 
+#if WITH_AUDIOMODULATION
 namespace AudioModulation
 {
 	class FAudioModulationSystem
