Split Modulation Registration into actual registration/Modulation proxy generation vs Settings construction to support AudioDevice-less DATAProxy generation
- Made modulator settings objects stateless & removed AudioDevice dependency in mod settings constructors.
- Required moving ModulationParameters to thread-safe registered singleton registry in AudioExtensions and out of plugin
- Rename IAudioModulation/FAudioModulation to IAudioModulationManager/FAudioModulationManager for clarity
#rnx

diff --git a/Engine/Plugins/Runtime/AudioModulation/Source/AudioModulation/Private/AudioModulationSystem.h b/Engine/Plugins/Runtime/AudioModulation/Source/AudioModulation/Private/AudioModulationSystem.h
--- a/Engine/Plugins/Runtime/AudioModulation/Source/AudioModulation/Private/AudioModulationSystem.h
+++ b/Engine/Plugins/Runtime/AudioModulation/Source/AudioModulation/Private/AudioModulationSystem.h
@@ -58,6 +58,11 @@ namespace AudioModulation
 	};
 } // namespace AudioModulation
 
+#if WITH_AUDIOMODULATION
+
+#if !UE_BUILD_SHIPPING
+#include "AudioModulationDebugger.h"
+#endif // !UE_BUILD_SHIPPING
 
 #if WITH_AUDIOMODULATION
 namespace AudioModulation
