Fix audio modulation proxy using unit space instead of normalized space.

diff --git a/Engine/Plugins/Runtime/AudioModulation/Source/AudioModulation/Private/SoundModulationPatchProxy.cpp b/Engine/Plugins/Runtime/AudioModulation/Source/AudioModulation/Private/SoundModulationPatchProxy.cpp
--- a/Engine/Plugins/Runtime/AudioModulation/Source/AudioModulation/Private/SoundModulationPatchProxy.cpp
+++ b/Engine/Plugins/Runtime/AudioModulation/Source/AudioModulation/Private/SoundModulationPatchProxy.cpp
@@ -45,6 +45,10 @@ namespace AudioModulation
 	{
 		bBypass = InSettings.bBypass;
 		DefaultValue = InSettings.OutputParameter.DefaultValue;
+		if (InSettings.OutputParameter.bRequiresConversion)
+		{
+			InSettings.OutputParameter.NormalizedFunction(DefaultValue);
+		}
 
 		// Cache existing proxies to avoid releasing bus state (and potentially referenced bus state) when reinitializing
 		const TArray<FModulationInputProxy> CachedProxies = InputProxies;
