Audio modulation fixes for generator bypass and pitch modulation

diff --git a/Engine/Plugins/Runtime/AudioModulation/Source/AudioModulation/Private/Generators/SoundModulationLFO.cpp b/Engine/Plugins/Runtime/AudioModulation/Source/AudioModulation/Private/Generators/SoundModulationLFO.cpp
--- a/Engine/Plugins/Runtime/AudioModulation/Source/AudioModulation/Private/Generators/SoundModulationLFO.cpp
+++ b/Engine/Plugins/Runtime/AudioModulation/Source/AudioModulation/Private/Generators/SoundModulationLFO.cpp
@@ -86,14 +86,15 @@ namespace AudioModulation
 				AudioRenderThreadCommand([this, NewGenerator = MoveTemp(InGenerator)]()
 				{
 					const FLFOGenerator* Generator = static_cast<const FLFOGenerator*>(NewGenerator.Get());
+					Params = Generator->Params;
 
-					LFO.SetGain(Generator->Params.Amplitude);
-					LFO.SetFrequency(Generator->Params.Frequency);
+					LFO.SetGain(Params.Amplitude);
+					LFO.SetFrequency(Params.Frequency);
 
 					static_assert(static_cast<int32>(ESoundModulationLFOShape::COUNT) == static_cast<int32>(Audio::ELFO::Type::NumLFOTypes), "LFOShape/ELFO Type mismatch");
-					LFO.SetType(static_cast<Audio::ELFO::Type>(Generator->Params.Shape));
+					LFO.SetType(static_cast<Audio::ELFO::Type>(Params.Shape));
 
-					Audio::ELFOMode::Type NewMode = Generator->Params.bLooping ? Audio::ELFOMode::Type::Sync : Audio::ELFOMode::OneShot;
+					Audio::ELFOMode::Type NewMode = Params.bLooping ? Audio::ELFOMode::Type::Sync : Audio::ELFOMode::OneShot;
 					const bool bModeUpdated = NewMode != LFO.GetMode();
 					if (bModeUpdated)
 					{
diff --git a/Engine/Source/Runtime/AudioMixer/Private/AudioMixerSourceManager.cpp b/Engine/Source/Runtime/AudioMixer/Private/AudioMixerSourceManager.cpp
--- a/Engine/Source/Runtime/AudioMixer/Private/AudioMixerSourceManager.cpp
+++ b/Engine/Source/Runtime/AudioMixer/Private/AudioMixerSourceManager.cpp
@@ -744,7 +744,6 @@ namespace Audio
 		PitchMod.UpdateModulator(InitParams.ModulationSettings.PitchModulationDestination.Modulator);
 
 		FModulationDestination HighpassMod;
-		PitchMod.Init(MixerDevice->DeviceID, FName("Pitch"), false /* bInIsBuffered */);
 		HighpassMod.Init(MixerDevice->DeviceID, FName("HPFCutoffFrequency"), false /* bInIsBuffered */);
 		HighpassMod.UpdateModulator(InitParams.ModulationSettings.HighpassModulationDestination.Modulator);
 
