Fix for parameter proxies not registering missing parameters on demand
#rnx

diff --git a/Engine/Plugins/Runtime/AudioModulation/Source/AudioModulation/Private/SoundControlBus.cpp b/Engine/Plugins/Runtime/AudioModulation/Source/AudioModulation/Private/SoundControlBus.cpp
--- a/Engine/Plugins/Runtime/AudioModulation/Source/AudioModulation/Private/SoundControlBus.cpp
+++ b/Engine/Plugins/Runtime/AudioModulation/Source/AudioModulation/Private/SoundControlBus.cpp
@@ -141,6 +141,6 @@ TUniquePtr<Audio::IProxyData> USoundControlBus::CreateNewProxyData(const Audio::
 
 const Audio::FModulationParameter& USoundControlBus::GetOutputParameter() const
 {
-	return AudioModulation::GetOrRegisterParameter(Parameter, *this);
+	const FString Breadcrumb = FString::Format(TEXT("{0} '{1}'"), { *GetClass()->GetName(), *GetName() });
+	return AudioModulation::GetOrRegisterParameter(Parameter, Breadcrumb);
 }
-
diff --git a/Engine/Plugins/Runtime/AudioModulation/Source/AudioModulation/Private/SoundModulationParameter.cpp b/Engine/Plugins/Runtime/AudioModulation/Source/AudioModulation/Private/SoundModulationParameter.cpp
--- a/Engine/Plugins/Runtime/AudioModulation/Source/AudioModulation/Private/SoundModulationParameter.cpp
+++ b/Engine/Plugins/Runtime/AudioModulation/Source/AudioModulation/Private/SoundModulationParameter.cpp
@@ -224,7 +224,7 @@ float USoundModulationParameterVolume::GetUnitMax() const
 
 namespace AudioModulation
 {
-	const Audio::FModulationParameter& GetOrRegisterParameter(const USoundModulationParameter* InParameter, const USoundModulatorBase& InReferencingModulator)
+	const Audio::FModulationParameter& GetOrRegisterParameter(const USoundModulationParameter* InParameter, const FString& InBreadcrumb)
 	{
 		FName ParamName;
 		if (InParameter)
@@ -233,9 +233,9 @@ namespace AudioModulation
 			if (!Audio::IsModulationParameterRegistered(ParamName))
 			{
 				UE_LOG(LogAudioModulation, Display,
-					TEXT("Parameter '%s' not registered.  Registration forced via activating referrencing modulator '%s'."),
+					TEXT("Parameter '%s' not registered.  Registration forced via '%s'."),
 					*ParamName.ToString(),
-					*InReferencingModulator.GetName());
+					*InBreadcrumb);
 			}
 
 			Audio::RegisterModulationParameter(ParamName, InParameter->CreateParameter());
@@ -247,18 +247,7 @@ namespace AudioModulation
 
 	FSoundModulationPluginParameterAssetProxy::FSoundModulationPluginParameterAssetProxy(USoundModulationParameter* InParameter)
 	{
-		using namespace Audio;
-
-		if (InParameter)
-		{
-			const FName ParameterName = InParameter->GetFName();
-			Parameter = Audio::GetModulationParameter(ParameterName);
-
-			if (ParameterName != Parameter.ParameterName)
-			{
-				UE_LOG(LogAudioModulation, Log, TEXT("Parameter '%s' not registered. Using default parameter."), *ParameterName.ToString());
-			}
-		}
+		Parameter = GetOrRegisterParameter(InParameter, TEXT("FSoundModulationPluginParameterAssetProxy construction"));
 	}
 
 	Audio::IProxyDataPtr FSoundModulationPluginParameterAssetProxy::Clone() const
diff --git a/Engine/Plugins/Runtime/AudioModulation/Source/AudioModulation/Private/SoundModulationPatch.cpp b/Engine/Plugins/Runtime/AudioModulation/Source/AudioModulation/Private/SoundModulationPatch.cpp
--- a/Engine/Plugins/Runtime/AudioModulation/Source/AudioModulation/Private/SoundModulationPatch.cpp
+++ b/Engine/Plugins/Runtime/AudioModulation/Source/AudioModulation/Private/SoundModulationPatch.cpp
@@ -34,7 +34,8 @@ TUniquePtr<Audio::IModulatorSettings> USoundModulationPatch::CreateProxySettings
 
 const Audio::FModulationParameter& USoundModulationPatch::GetOutputParameter() const
 {
-	return AudioModulation::GetOrRegisterParameter(PatchSettings.OutputParameter, *this);
+	const FString Breadcrumb = FString::Format(TEXT("{0} '{1}'"), { *GetClass()->GetName(), *GetName() });
+	return AudioModulation::GetOrRegisterParameter(PatchSettings.OutputParameter, Breadcrumb);
 }
 
 #if WITH_EDITOR
diff --git a/Engine/Plugins/Runtime/AudioModulation/Source/AudioModulation/Public/SoundModulationParameter.h b/Engine/Plugins/Runtime/AudioModulation/Source/AudioModulation/Public/SoundModulationParameter.h
--- a/Engine/Plugins/Runtime/AudioModulation/Source/AudioModulation/Public/SoundModulationParameter.h
+++ b/Engine/Plugins/Runtime/AudioModulation/Source/AudioModulation/Public/SoundModulationParameter.h
@@ -265,5 +265,10 @@ namespace AudioModulation
 		virtual const Audio::FModulationParameter& GetParameter() const override;
 	};
 
-	AUDIOMODULATION_API const Audio::FModulationParameter& GetOrRegisterParameter(const USoundModulationParameter* InParameter, const USoundModulatorBase& InReferencingModulator);
+	/*
+	 * Returns given registered parameter instance reference or creates it from the given asset if not registered.
+	 * @param InParameter - Parameter asset associated with the pre-existing or to-create parameter
+	 * @param InBreadcrumb - String identifying get or register request initiator.
+	 */
+	AUDIOMODULATION_API const Audio::FModulationParameter& GetOrRegisterParameter(const USoundModulationParameter* InParameter, const FString& InBreadcrumb);
 } // namespace AudioModulation
\ No newline at end of file
