- Fix upcast copy causing parameters to not be cached properly on proxy by removing unnecessary virtual function
- Fix return value improperly assigning value in MixModulatorNode, causing it always to be 0 or 1 based on if the modulator was registered (as return is if the value was found and assigned, and its a boolean cast when the modulator is found)
- Remove unused parameter proxy ModHandle member
- Don't re-register if parameter already registered
- Remove duplicate validation call
#rnx

diff --git a/Engine/Plugins/Runtime/AudioModulation/Source/AudioModulation/Private/MetaSoundNodes/MixModulatorsNode.cpp b/Engine/Plugins/Runtime/AudioModulation/Source/AudioModulation/Private/MetaSoundNodes/MixModulatorsNode.cpp
--- a/Engine/Plugins/Runtime/AudioModulation/Source/AudioModulation/Private/MetaSoundNodes/MixModulatorsNode.cpp
+++ b/Engine/Plugins/Runtime/AudioModulation/Source/AudioModulation/Private/MetaSoundNodes/MixModulatorsNode.cpp
@@ -134,7 +134,7 @@ namespace AudioModulation
 
 				if (OutHandle.IsValid())
 				{
-					Value = OutHandle.GetValueThreadSafe(Value);
+					OutHandle.GetValueThreadSafe(Value);
 				}
 
 				return Value;
diff --git a/Engine/Plugins/Runtime/AudioModulation/Source/AudioModulation/Private/SoundModulationParameter.cpp b/Engine/Plugins/Runtime/AudioModulation/Source/AudioModulation/Private/SoundModulationParameter.cpp
--- a/Engine/Plugins/Runtime/AudioModulation/Source/AudioModulation/Private/SoundModulationParameter.cpp
+++ b/Engine/Plugins/Runtime/AudioModulation/Source/AudioModulation/Private/SoundModulationParameter.cpp
@@ -15,6 +15,20 @@ TUniquePtr<Audio::IProxyData> USoundModulationParameter::CreateNewProxyData(cons
 }
 
 #if WITH_EDITOR
+void USoundModulationParameter::PostEditChangeProperty(struct FPropertyChangedEvent& PropertyChangedEvent)
+{
+	if (PropertyChangedEvent.ChangeType != EPropertyChangeType::Interactive)
+	{
+		const FName AssetName = GetFName();
+		if (Audio::IsModulationParameterRegistered(AssetName))
+		{
+			Audio::FModulationParameter NewParam = CreateParameter();
+			Audio::RegisterModulationParameter(NewParam.ParameterName, MoveTemp(NewParam));
+		}
+	}
+}
+
+
 void USoundModulationParameter::RefreshNormalizedValue()
 {
 	const float NewNormalizedValue = ConvertUnitToNormalized(Settings.ValueUnit);
@@ -236,9 +250,9 @@ namespace AudioModulation
 					TEXT("Parameter '%s' not registered.  Registration forced via '%s'."),
 					*ParamName.ToString(),
 					*InBreadcrumb);
-			}
 
-			Audio::RegisterModulationParameter(ParamName, InParameter->CreateParameter());
+				Audio::RegisterModulationParameter(ParamName, InParameter->CreateParameter());
+			}
 		}
 
 		// Returns default modulation parameter if no parameter provided.
@@ -254,9 +268,4 @@ namespace AudioModulation
 	{
 		return TUniquePtr<FSoundModulationPluginParameterAssetProxy>(new FSoundModulationPluginParameterAssetProxy(*this));
 	}
-
-	const Audio::FModulationParameter& FSoundModulationPluginParameterAssetProxy::GetParameter() const
-	{
-		return Parameter;
-	}
 }
\ No newline at end of file
diff --git a/Engine/Plugins/Runtime/AudioModulation/Source/AudioModulation/Public/SoundModulationParameter.h b/Engine/Plugins/Runtime/AudioModulation/Source/AudioModulation/Public/SoundModulationParameter.h
--- a/Engine/Plugins/Runtime/AudioModulation/Source/AudioModulation/Public/SoundModulationParameter.h
+++ b/Engine/Plugins/Runtime/AudioModulation/Source/AudioModulation/Public/SoundModulationParameter.h
@@ -23,7 +23,7 @@ struct AUDIOMODULATION_API FSoundModulationParameterSettings
 	  * If GetMixFunction performs the mathematical operation f(x1, x2), then the default ValueNormalized should result in
 	  * f(x1, d) = x1 where d is ValueNormalized.
 	  */
-	UPROPERTY(EditAnywhere, Category = General, meta = (ClampMin = "0.0", ClampMax = "1.0"))
+	UPROPERTY(EditAnywhere, Category = General, BlueprintReadOnly, meta = (ClampMin = "0.0", ClampMax = "1.0"))
 	float ValueNormalized = 1.0f;
 
 #if WITH_EDITORONLY_DATA
@@ -108,12 +108,13 @@ public:
 	//~ End IAudioProxyDataFactory Interface.
 
 #if WITH_EDITOR
+	virtual void PostEditChangeProperty(struct FPropertyChangedEvent& PropertyChangedEvent) override;
+
 	void RefreshNormalizedValue();
 	void RefreshUnitValue();
 #endif // WITH_EDITOR
 
 	Audio::FModulationParameter CreateParameter() const;
-
 };
 
 // Linearly scaled value between unit minimum and maximum.
@@ -124,11 +125,11 @@ class USoundModulationParameterScaled : public USoundModulationParameter
 
 public:
 	/** Unit minimum of modulator. Minimum is only enforced at modulation destination. */
-	UPROPERTY(EditAnywhere, Category = General)
+	UPROPERTY(EditAnywhere, Category = General, BlueprintReadOnly)
 	float UnitMin = 0.0f;
 
 	/** Unit maximum of modulator. Maximum is only enforced at modulation destination. */
-	UPROPERTY(EditAnywhere, Category = General)
+	UPROPERTY(EditAnywhere, Category = General, BlueprintReadOnly)
 	float UnitMax = 1.0f;
 
 	virtual bool RequiresUnitConversion() const override;
@@ -158,11 +159,11 @@ class USoundModulationParameterFrequency : public USoundModulationParameterFrequ
 
 public:
 	/** Unit minimum of modulator. Minimum is only enforced at modulation destination. */
-	UPROPERTY(EditAnywhere, Category = General)
+	UPROPERTY(EditAnywhere, Category = General, BlueprintReadOnly)
 	float UnitMin = MIN_FILTER_FREQUENCY;
 
 	/** Unit maximum of modulator. Maximum is only enforced at modulation destination. */
-	UPROPERTY(EditAnywhere, Category = General)
+	UPROPERTY(EditAnywhere, Category = General, BlueprintReadOnly)
 	float UnitMax = MAX_FILTER_FREQUENCY;
 
 	virtual float GetUnitMin() const override
@@ -224,7 +225,7 @@ class USoundModulationParameterBipolar : public USoundModulationParameter
 
 public:
 	/** Unit range of modulator. Range is only enforced at modulation destination. */
-	UPROPERTY(EditAnywhere, Category = General, meta = (ClampMin = 0.00000001))
+	UPROPERTY(EditAnywhere, Category = General, BlueprintReadOnly, meta = (ClampMin = 0.00000001))
 	float UnitRange = 2.0f;
 
 	virtual bool RequiresUnitConversion() const override;
@@ -242,7 +243,7 @@ class USoundModulationParameterVolume : public USoundModulationParameter
 
 public:
 	/** Minimum volume of parameter. Only enforced at modulation destination. */
-	UPROPERTY(EditAnywhere, Category = General, meta = (ClampMax = 0.0))
+	UPROPERTY(EditAnywhere, Category = General, BlueprintReadOnly, meta = (ClampMax = 0.0))
 	float MinVolume = -100.0f;
 
 	virtual bool RequiresUnitConversion() const override;
@@ -261,8 +262,6 @@ namespace AudioModulation
 		FSoundModulationPluginParameterAssetProxy(const FSoundModulationPluginParameterAssetProxy& InProxy) = default;
 
 		virtual Audio::IProxyDataPtr Clone() const override;
-
-		virtual const Audio::FModulationParameter& GetParameter() const override;
 	};
 
 	/*
diff --git a/Engine/Plugins/Runtime/Metasound/Source/MetasoundEditor/Private/MetasoundEditorGraph.cpp b/Engine/Plugins/Runtime/Metasound/Source/MetasoundEditor/Private/MetasoundEditorGraph.cpp
--- a/Engine/Plugins/Runtime/Metasound/Source/MetasoundEditor/Private/MetasoundEditorGraph.cpp
+++ b/Engine/Plugins/Runtime/Metasound/Source/MetasoundEditor/Private/MetasoundEditorGraph.cpp
@@ -1682,12 +1682,6 @@ bool UMetasoundEditorGraph::ValidateInternal(Metasound::Editor::FGraphValidation
 	{
 		FGraphNodeValidationResult NodeResult(*Node);
 
-		if (UMetasoundEditorGraphExternalNode* ExternalNode = Cast<UMetasoundEditorGraphExternalNode>(Node))
-		{
-			bIsValid &= ExternalNode->Validate(NodeResult);
-			bMarkDirty |= NodeResult.bIsDirty;
-		}
-
 		bIsValid &= Node->Validate(NodeResult);
 		bMarkDirty |= NodeResult.bIsDirty;
 
diff --git a/Engine/Source/Runtime/AudioExtensions/Public/IAudioModulation.h b/Engine/Source/Runtime/AudioExtensions/Public/IAudioModulation.h
--- a/Engine/Source/Runtime/AudioExtensions/Public/IAudioModulation.h
+++ b/Engine/Source/Runtime/AudioExtensions/Public/IAudioModulation.h
@@ -253,12 +253,10 @@ public:
 
 	virtual const Audio::FModulationParameter& GetParameter() const
 	{
-		static const Audio::FModulationParameter DefaultParam;
-		return DefaultParam;
+		return Parameter;
 	}
 
-	protected:
-	Audio::FModulatorHandle ModHandle;
+protected:
 	Audio::FModulationParameter Parameter;
 };
 using FSoundModulationParameterAssetProxyPtr = TSharedPtr<FSoundModulationParameterAssetProxy, ESPMode::ThreadSafe>;
