Clarify when audio modulation buses and bus mixes may be destroyed by garbage collection

#rnx

diff --git a/Engine/Plugins/Runtime/AudioModulation/Source/AudioModulation/Private/SoundControlBus.cpp b/Engine/Plugins/Runtime/AudioModulation/Source/AudioModulation/Private/SoundControlBus.cpp
--- a/Engine/Plugins/Runtime/AudioModulation/Source/AudioModulation/Private/SoundControlBus.cpp
+++ b/Engine/Plugins/Runtime/AudioModulation/Source/AudioModulation/Private/SoundControlBus.cpp
@@ -105,8 +105,6 @@ void USoundControlBus::BeginDestroy()
 {
 	using namespace AudioModulation;
 
-	Super::BeginDestroy();
-
 	if (UWorld* World = GetWorld())
 	{
 		FAudioDeviceHandle AudioDevice = World->GetAudioDevice();
@@ -121,6 +119,9 @@ void USoundControlBus::BeginDestroy()
 			}
 		}
 	}
+
+	// Call parent destroy at end to ensure object is in a valid state for the modulation manager to clean up first
+	Super::BeginDestroy();
 }
 
 const Audio::FModulationMixFunction USoundControlBus::GetMixFunction() const
diff --git a/Engine/Plugins/Runtime/AudioModulation/Source/AudioModulation/Private/SoundControlBusMix.cpp b/Engine/Plugins/Runtime/AudioModulation/Source/AudioModulation/Private/SoundControlBusMix.cpp
--- a/Engine/Plugins/Runtime/AudioModulation/Source/AudioModulation/Private/SoundControlBusMix.cpp
+++ b/Engine/Plugins/Runtime/AudioModulation/Source/AudioModulation/Private/SoundControlBusMix.cpp
@@ -39,8 +39,6 @@ void USoundControlBusMix::BeginDestroy()
 {
 	using namespace AudioModulation;
 
-	Super::BeginDestroy();
-
 	if (UWorld* World = GetWorld())
 	{
 		FAudioDeviceHandle AudioDevice = World->GetAudioDevice();
@@ -57,6 +55,9 @@ void USoundControlBusMix::BeginDestroy()
 			}
 		}
 	}
+
+	// Call parent destroy at end to ensure object is in a valid state for the modulation manager to clean up first
+	Super::BeginDestroy();
 }
 
 void USoundControlBusMix::ActivateMix()
diff --git a/Engine/Plugins/Runtime/AudioModulation/Source/AudioModulation/Public/AudioModulationStatics.h b/Engine/Plugins/Runtime/AudioModulation/Source/AudioModulation/Public/AudioModulationStatics.h
--- a/Engine/Plugins/Runtime/AudioModulation/Source/AudioModulation/Public/AudioModulationStatics.h
+++ b/Engine/Plugins/Runtime/AudioModulation/Source/AudioModulation/Public/AudioModulationStatics.h
@@ -72,13 +72,14 @@ public:
 	 * @param Activate - Whether or not to activate bus on creation. If true, deactivation will only occur
 	 * if returned bus is manually deactivated and not referenced or destroyed (i.e. will not deactivate
 	 * when all references become inactive).
+	 * @return Capture this in a Blueprint variable to prevent it from being automatically garbage collected. 
 	 */
 	UFUNCTION(BlueprintCallable, Category = "Audio", DisplayName = "Create Control Bus", meta = (
 		AdvancedDisplay = "3",
 		WorldContext = "WorldContextObject",
 		Keywords = "make modulation LPF modulator")
 	)
-	static USoundControlBus* CreateBus(UObject* WorldContextObject, FName Name, USoundModulationParameter* Parameter, bool Activate = true);
+	static UPARAM(DisplayName = "Bus") USoundControlBus* CreateBus(UObject* WorldContextObject, FName Name, USoundModulationParameter* Parameter, bool Activate = true);
 
 	/** Creates a stage used to mix a control bus.
 	 * @param Bus - Bus stage is in charge of applying mix value to.
@@ -104,12 +105,13 @@ public:
 	 * @param Activate - Whether or not to activate mix on creation. If true, deactivation will only occur
 	 * if returned mix is manually deactivated and not referenced or destroyed (i.e. will not deactivate
 	 * when all references become inactive).
+	 * @return Capture this in a Blueprint variable to prevent it from being automatically garbage collected. 
 	 */
 	UFUNCTION(BlueprintCallable, Category = "Audio", DisplayName = "Create Control Bus Mix", meta = (
 		WorldContext = "WorldContextObject",
 		Keywords = "make modulation modulator")
 	)
-	static USoundControlBusMix* CreateBusMix(
+	static UPARAM(DisplayName = "BusMix") USoundControlBusMix* CreateBusMix(
 		UObject* WorldContextObject,
 		FName Name, 
 		TArray<FSoundControlBusMixStage> Stages,
