Add ability to compile out MetaSound support in AudioModulation

diff --git a/Engine/Plugins/Runtime/AudioModulation/Source/AudioModulation/AudioModulation.Build.cs b/Engine/Plugins/Runtime/AudioModulation/Source/AudioModulation/AudioModulation.Build.cs
--- a/Engine/Plugins/Runtime/AudioModulation/Source/AudioModulation/AudioModulation.Build.cs
+++ b/Engine/Plugins/Runtime/AudioModulation/Source/AudioModulation/AudioModulation.Build.cs
@@ -4,6 +4,12 @@ namespace UnrealBuildTool.Rules
 {
 	public class AudioModulation : ModuleRules
 	{
+		// Set this to false & disable MetaSound plugin dependency
+		// by setting MetaSound's field '"Enabled": false' in the
+		// AudioModulation.uplugin if running Modulation without
+		// MetaSound support.
+		public static bool bIncludeMetaSoundSupport = true;
+
 		public AudioModulation(ReadOnlyTargetRules Target) : base(Target)
 		{
 			PublicDependencyModuleNames.AddRange(
@@ -21,12 +27,21 @@ namespace UnrealBuildTool.Rules
 					"CoreUObject",
 					"DeveloperSettings",
 					"Engine",
-					"MetasoundFrontend",
-					"MetasoundGraphCore",
 					"SignalProcessing"
 				}
 			);
 
+			if (bIncludeMetaSoundSupport)
+			{
+				PrivateDependencyModuleNames.AddRange(
+					new string[]
+					{
+						"MetasoundFrontend",
+						"MetasoundGraphCore"
+					}
+				);
+			}
+
 			if (Target.Type == TargetType.Editor)
 			{
 				PrivateDependencyModuleNames.AddRange(
@@ -39,6 +54,14 @@ namespace UnrealBuildTool.Rules
 			}
 
 			PublicDefinitions.Add("WITH_AUDIOMODULATION=1");
+			if (bIncludeMetaSoundSupport)
+			{
+				PublicDefinitions.Add("WITH_AUDIOMODULATION_METASOUND_SUPPORT=1");
+			}
+			else
+			{
+				PublicDefinitions.Add("WITH_AUDIOMODULATION_METASOUND_SUPPORT=0");
+			}
 		}
 	}
 }
\ No newline at end of file
diff --git a/Engine/Plugins/Runtime/AudioModulation/Source/AudioModulation/Private/AudioModulation.cpp b/Engine/Plugins/Runtime/AudioModulation/Source/AudioModulation/Private/AudioModulation.cpp
--- a/Engine/Plugins/Runtime/AudioModulation/Source/AudioModulation/Private/AudioModulation.cpp
+++ b/Engine/Plugins/Runtime/AudioModulation/Source/AudioModulation/Private/AudioModulation.cpp
@@ -1,4 +1,5 @@
 // Copyright Epic Games, Inc. All Rights Reserved.
+
 #include "AudioModulation.h"
 
 #include "AudioModulationLogging.h"
@@ -7,8 +8,6 @@
 #include "CanvasTypes.h"
 #include "Features/IModularFeatures.h"
 #include "IAudioModulation.h"
-#include "MetasoundDataTypeRegistrationMacro.h"
-#include "MetasoundFrontendRegistries.h"
 #include "Modules/ModuleManager.h"
 #include "SoundControlBusMix.h"
 #include "SoundModulationParameter.h"
@@ -16,9 +15,13 @@
 #include "SoundModulatorAsset.h"
 #include "UObject/NoExportTypes.h"
 
+#if WITH_AUDIOMODULATION_METASOUND_SUPPORT
+#include "MetasoundDataTypeRegistrationMacro.h"
+#include "MetasoundFrontendRegistries.h"
 
 REGISTER_METASOUND_DATATYPE(AudioModulation::FSoundModulatorAsset, "Modulator", Metasound::ELiteralType::UObjectProxy, USoundModulatorBase);
 REGISTER_METASOUND_DATATYPE(AudioModulation::FSoundModulationParameterAsset, "ModulationParameter", Metasound::ELiteralType::UObjectProxy, USoundModulationParameter);
+#endif // WITH_AUDIOMODULATION_METASOUND_SUPPORT
 
 namespace AudioModulation
 {
@@ -260,8 +263,11 @@ void FAudioModulationModule::StartupModule()
 		ModulationSettings->RegisterParameters();
 	}
 
+#if WITH_AUDIOMODULATION_METASOUND_SUPPORT
+	UE_LOG(LogAudioModulation, Log, TEXT("Registering Modulation MetaSound Nodes..."));
 	// flush node registration queue to guarantee AudioModulation DataTypes/Nodes are ready prior to assets loading
 	FMetasoundFrontendRegistryContainer::Get()->RegisterPendingNodes();
+#endif // WITH_AUDIOMODULATION_METASOUND_SUPPORT
 
 	UE_LOG(LogAudioModulation, Log, TEXT("Audio Modulation Initialized"));
 }
diff --git a/Engine/Plugins/Runtime/AudioModulation/Source/AudioModulation/Private/MetaSoundNodes/GetModulatorValueNode.cpp b/Engine/Plugins/Runtime/AudioModulation/Source/AudioModulation/Private/MetaSoundNodes/GetModulatorValueNode.cpp
--- a/Engine/Plugins/Runtime/AudioModulation/Source/AudioModulation/Private/MetaSoundNodes/GetModulatorValueNode.cpp
+++ b/Engine/Plugins/Runtime/AudioModulation/Source/AudioModulation/Private/MetaSoundNodes/GetModulatorValueNode.cpp
@@ -1,5 +1,6 @@
 // Copyright Epic Games, Inc. All Rights Reserved.
 
+#if WITH_AUDIOMODULATION_METASOUND_SUPPORT
 #include "AudioDefines.h"
 #include "AudioModulation.h"
 #include "MetasoundDataFactory.h"
@@ -161,3 +162,4 @@ namespace AudioModulation
 } // namespace AudioModulation
 
 #undef LOCTEXT_NAMESPACE
+#endif // WITH_AUDIOMODULATION_METASOUND_SUPPORT
\ No newline at end of file
diff --git a/Engine/Plugins/Runtime/AudioModulation/Source/AudioModulation/Private/MetaSoundNodes/MixModulatorsNode.cpp b/Engine/Plugins/Runtime/AudioModulation/Source/AudioModulation/Private/MetaSoundNodes/MixModulatorsNode.cpp
--- a/Engine/Plugins/Runtime/AudioModulation/Source/AudioModulation/Private/MetaSoundNodes/MixModulatorsNode.cpp
+++ b/Engine/Plugins/Runtime/AudioModulation/Source/AudioModulation/Private/MetaSoundNodes/MixModulatorsNode.cpp
@@ -1,5 +1,7 @@
 // Copyright Epic Games, Inc. All Rights Reserved.
 
+#if WITH_AUDIOMODULATION_METASOUND_SUPPORT
+
 #include "AudioModulation.h"
 #include "IAudioModulation.h"
 #include "MetasoundDataFactory.h"
@@ -187,3 +189,4 @@ namespace AudioModulation
 } // namespace AudioModulation
 
 #undef LOCTEXT_NAMESPACE
+#endif // WITH_AUDIOMODULATION_METASOUND_SUPPORT
diff --git a/Engine/Plugins/Runtime/AudioModulation/Source/AudioModulation/Private/SoundModulatorAsset.cpp b/Engine/Plugins/Runtime/AudioModulation/Source/AudioModulation/Private/SoundModulatorAsset.cpp
--- a/Engine/Plugins/Runtime/AudioModulation/Source/AudioModulation/Private/SoundModulatorAsset.cpp
+++ b/Engine/Plugins/Runtime/AudioModulation/Source/AudioModulation/Private/SoundModulatorAsset.cpp
@@ -1,5 +1,6 @@
 // Copyright Epic Games, Inc. All Rights Reserved.
 
+#if WITH_AUDIOMODULATION_METASOUND_SUPPORT
 #include "SoundModulatorAsset.h"
 #include "Templates/SharedPointer.h"
 
@@ -38,3 +39,4 @@ namespace AudioModulation
 		}
 	}
 } // namespace AudioModulation
+#endif // WITH_AUDIOMODULATION_METASOUND_SUPPORT
diff --git a/Engine/Plugins/Runtime/AudioModulation/Source/AudioModulation/Public/SoundModulatorAsset.h b/Engine/Plugins/Runtime/AudioModulation/Source/AudioModulation/Public/SoundModulatorAsset.h
--- a/Engine/Plugins/Runtime/AudioModulation/Source/AudioModulation/Public/SoundModulatorAsset.h
+++ b/Engine/Plugins/Runtime/AudioModulation/Source/AudioModulation/Public/SoundModulatorAsset.h
@@ -1,11 +1,12 @@
 // Copyright Epic Games, Inc. All Rights Reserved.
 #pragma once
 
+#if WITH_AUDIOMODULATION_METASOUND_SUPPORT
 #include "IAudioModulation.h"
 #include "IAudioProxyInitializer.h"
 #include "Internationalization/Text.h"
-#include "MetasoundDataReferenceMacro.h"
 
+#include "MetasoundDataReferenceMacro.h"
 
 namespace AudioModulation
 {
@@ -106,3 +107,4 @@ namespace AudioModulation
 
 DECLARE_METASOUND_DATA_REFERENCE_TYPES(AudioModulation::FSoundModulatorAsset, AUDIOMODULATION_API, FSoundModulatorAssetTypeInfo, FSoundModulatorAssetReadRef, FSoundModulatorAssetWriteRef)
 DECLARE_METASOUND_DATA_REFERENCE_TYPES(AudioModulation::FSoundModulationParameterAsset, AUDIOMODULATION_API, FSoundModulationParameterAssetTypeInfo, FSoundModulationParameterAssetReadRef, FSoundModulationParameterAssetWriteRef)
+#endif // WITH_AUDIOMODULATION_METASOUND_SUPPORT
