Potential fix for missing class definition when asset manager loads modulation parameters
#rnx

diff --git a/Engine/Plugins/Runtime/AudioModulation/Source/AudioModulation/Public/AudioModulation.h b/Engine/Plugins/Runtime/AudioModulation/Source/AudioModulation/Public/AudioModulation.h
--- a/Engine/Plugins/Runtime/AudioModulation/Source/AudioModulation/Public/AudioModulation.h
+++ b/Engine/Plugins/Runtime/AudioModulation/Source/AudioModulation/Public/AudioModulation.h
@@ -92,7 +92,7 @@ namespace AudioModulation
 	AUDIOMODULATION_API void IterateModulationManagers(TFunctionRef<void(FAudioModulationManager&)> InFunction);
 } // namespace AudioModulation
 
-class FAudioModulationPluginFactory : public IAudioModulationFactory
+class AUDIOMODULATION_API FAudioModulationPluginFactory : public IAudioModulationFactory
 {
 public:
 	virtual const FName& GetDisplayName() const override
@@ -104,7 +104,7 @@ public:
 	virtual TAudioModulationPtr CreateNewModulationPlugin(FAudioDevice* OwningDevice) override;
 };
 
-class FAudioModulationModule : public IModuleInterface
+class AUDIOMODULATION_API FAudioModulationModule : public IModuleInterface
 {
 public:
 	virtual void StartupModule() override;
diff --git a/Engine/Plugins/Runtime/AudioModulation/Source/AudioModulationEditor/Private/AudioModulationEditor.cpp b/Engine/Plugins/Runtime/AudioModulation/Source/AudioModulationEditor/Private/AudioModulationEditor.cpp
--- a/Engine/Plugins/Runtime/AudioModulation/Source/AudioModulationEditor/Private/AudioModulationEditor.cpp
+++ b/Engine/Plugins/Runtime/AudioModulation/Source/AudioModulationEditor/Private/AudioModulationEditor.cpp
@@ -7,6 +7,7 @@
 #include "AssetTypeActions/AssetTypeActions_SoundModulationGenerator.h"
 #include "AssetTypeActions/AssetTypeActions_SoundModulationParameter.h"
 #include "AssetTypeActions/AssetTypeActions_SoundModulationPatch.h"
+#include "AudioModulation.h"
 #include "Editors/ModulationPatchCurveEditorViewStacked.h"
 #include "Framework/Commands/UIAction.h"
 #include "Framework/MultiBox/MultiBoxBuilder.h"
@@ -74,6 +75,9 @@ void FAudioModulationEditorModule::SetIcon(const FString& ClassName)
 
 void FAudioModulationEditorModule::StartupModule()
 {
+	// Seems to be required to avoid missing class definition on start-up for some platforms when the asset manager loads
+	FModuleManager::LoadModuleChecked<FAudioModulationModule>("AudioModulation");
+
 	static const FName AudioModulationStyleName(TEXT("AudioModulationStyle"));
 	StyleSet = MakeShared<FSlateStyleSet>(AudioModulationStyleName);
 
@@ -110,7 +114,6 @@ void FAudioModulationEditorModule::StartupModule()
 	// All parameters are required to always be loaded in editor to enable them to be referenced via object
 	// metadata and custom layouts, even if they are not referenced by runtime uobjects/systems directly
 	IAssetRegistry& AssetRegistry = FModuleManager::LoadModuleChecked<FAssetRegistryModule>("AssetRegistry").Get();
-	TArray<FAssetData> Assets;
 	AssetRegistry.OnAssetAdded().AddLambda([](const FAssetData& InAssetData)
 	{
 		if (InAssetData.GetClass()->IsChildOf<USoundModulationParameter>())
