Fix undefined behaviour calling IsChildOf with a nullptr

diff --git a/Engine/Plugins/Runtime/AudioModulation/Source/AudioModulationEditor/Private/AudioModulationEditor.cpp b/Engine/Plugins/Runtime/AudioModulation/Source/AudioModulationEditor/Private/AudioModulationEditor.cpp
--- a/Engine/Plugins/Runtime/AudioModulation/Source/AudioModulationEditor/Private/AudioModulationEditor.cpp
+++ b/Engine/Plugins/Runtime/AudioModulation/Source/AudioModulationEditor/Private/AudioModulationEditor.cpp
@@ -116,7 +116,8 @@ void FAudioModulationEditorModule::StartupModule()
 	IAssetRegistry& AssetRegistry = FModuleManager::LoadModuleChecked<FAssetRegistryModule>("AssetRegistry").Get();
 	AssetRegistry.OnAssetAdded().AddLambda([](const FAssetData& InAssetData)
 	{
-		if (InAssetData.GetClass()->IsChildOf<USoundModulationParameter>())
+		UClass* AssetDataClass = InAssetData.GetClass();
+		if (AssetDataClass && AssetDataClass->IsChildOf<USoundModulationParameter>())
 		{
 			if (USoundModulationParameter* Parameter = CastChecked<USoundModulationParameter>(InAssetData.GetAsset()))
 			{
