From prior CL just broken out:
- Enable ability to lazy load type data for MetaSound PinType editor data
    - Utilize ability for Modulation types
- Remove Editor Frontend Class Metadata in favor of direct access from FrontendRegistry

diff --git a/Engine/Plugins/Runtime/AudioModulation/Source/AudioModulationEditor/AudioModulationEditor.build.cs b/Engine/Plugins/Runtime/AudioModulation/Source/AudioModulationEditor/AudioModulationEditor.build.cs
--- a/Engine/Plugins/Runtime/AudioModulation/Source/AudioModulationEditor/AudioModulationEditor.build.cs
+++ b/Engine/Plugins/Runtime/AudioModulation/Source/AudioModulationEditor/AudioModulationEditor.build.cs
@@ -4,10 +4,17 @@ using UnrealBuildTool;
 
 public class AudioModulationEditor : ModuleRules
 {
+	// Set this to false & disable MetaSound plugin dependency
+	// by setting MetaSound's field '"Enabled": false' in the
+	// AudioModulation.uplugin if running Modulation without
+	// MetaSound support.
+	public static bool bIncludeMetaSoundSupport = true;
+
 	public AudioModulationEditor(ReadOnlyTargetRules Target) : base(Target)
 	{
 		PrivateDependencyModuleNames.AddRange(
-			new string[] {
+			new string[]
+			{
 				"Core",
 				"CoreUObject",
 				"Engine",
@@ -29,10 +36,23 @@ public class AudioModulationEditor : ModuleRules
 		);
 
 		PrivateIncludePaths.AddRange(
-			new string[] {
+			new string[]
+			{
 				"AudioModulationEditor/Private",
 				"AudioModulation/Private"
 			}
 		);
+
+		if (bIncludeMetaSoundSupport)
+		{
+			PrivateDependencyModuleNames.AddRange(
+				new string[]
+				{
+					"MetasoundGraphCore",
+					"MetasoundFrontend",
+					"MetasoundEditor"
+				}
+			);
+		}
 	}
 }
diff --git a/Engine/Plugins/Runtime/AudioModulation/Source/AudioModulationEditor/Private/AudioModulationEditor.cpp b/Engine/Plugins/Runtime/AudioModulation/Source/AudioModulationEditor/Private/AudioModulationEditor.cpp
--- a/Engine/Plugins/Runtime/AudioModulation/Source/AudioModulationEditor/Private/AudioModulationEditor.cpp
+++ b/Engine/Plugins/Runtime/AudioModulation/Source/AudioModulationEditor/Private/AudioModulationEditor.cpp
@@ -2,6 +2,7 @@
 #include "AudioModulationEditor.h"
 
 #include "AssetRegistry/AssetRegistryModule.h"
+#include "AssetRegistry/IAssetRegistry.h"
 #include "AssetTypeActions/AssetTypeActions_SoundControlBus.h"
 #include "AssetTypeActions/AssetTypeActions_SoundControlBusMix.h"
 #include "AssetTypeActions/AssetTypeActions_SoundModulationGenerator.h"
@@ -12,7 +13,7 @@
 #include "Framework/Commands/UIAction.h"
 #include "Framework/MultiBox/MultiBoxBuilder.h"
 #include "Framework/MultiBox/MultiBoxExtender.h"
-#include "AssetRegistry/IAssetRegistry.h"
+#include "IAudioModulation.h"
 #include "ICurveEditorModule.h"
 #include "Internationalization/Internationalization.h"
 #include "Layouts/SoundControlBusMixStageLayout.h"
@@ -20,6 +21,7 @@
 #include "Layouts/SoundModulationParameterSettingsLayout.h"
 #include "Layouts/SoundModulationTransformLayout.h"
 #include "LevelEditor.h"
+#include "Modules/ModuleManager.h"
 #include "Sound/SoundBase.h"
 #include "SoundModulationParameter.h"
 #include "SoundModulationTransform.h"
@@ -27,9 +29,14 @@
 #include "Textures/SlateIcon.h"
 #include "UObject/UObjectIterator.h"
 
+#if WITH_AUDIOMODULATION_METASOUND_SUPPORT
+#include "MetasoundDataReference.h"
+#include "MetasoundEditorModule.h"
+#endif // WITH_AUDIOMODULATION_METASOUND_SUPPORT
 
 DEFINE_LOG_CATEGORY(LogAudioModulationEditor);
 
+
 namespace AudioModulationEditor
 {
 	static const FName ToolName = TEXT("AssetTools");
@@ -111,6 +118,18 @@ void FAudioModulationEditorModule::StartupModule()
 
 	FSlateStyleRegistry::RegisterSlateStyle(*StyleSet.Get());
 
+#if WITH_AUDIOMODULATION_METASOUND_SUPPORT
+	{
+		using namespace Metasound;
+		using namespace Metasound::Editor;
+		IMetasoundEditorModule& MetaSoundEditorModule = FModuleManager::GetModuleChecked<IMetasoundEditorModule>("MetaSoundEditor");
+		MetaSoundEditorModule.RegisterPinType("Modulator");
+		MetaSoundEditorModule.RegisterPinType(CreateArrayTypeNameFromElementTypeName("Modulator"));
+		MetaSoundEditorModule.RegisterPinType("ModulationParameter");
+		MetaSoundEditorModule.RegisterPinType(CreateArrayTypeNameFromElementTypeName("ModulationParameter"));
+	}
+#endif // WITH_AUDIOMODULATION_METASOUND_SUPPORT
+
 	// All parameters are required to always be loaded in editor to enable them to be referenced via object
 	// metadata and custom layouts, even if they are not referenced by runtime uobjects/systems directly
 	IAssetRegistry& AssetRegistry = FModuleManager::LoadModuleChecked<FAssetRegistryModule>("AssetRegistry").Get();
