Fix crash due to module load order

diff --git a/Engine/Plugins/Runtime/AudioModulation/AudioModulation.uplugin b/Engine/Plugins/Runtime/AudioModulation/AudioModulation.uplugin
--- a/Engine/Plugins/Runtime/AudioModulation/AudioModulation.uplugin
+++ b/Engine/Plugins/Runtime/AudioModulation/AudioModulation.uplugin
@@ -19,7 +19,7 @@
 		{
 			"Name": "AudioModulation",
 			"Type": "Runtime",
-			"LoadingPhase": "Default"
+			"LoadingPhase": "PreDefault"
 		},
 		{
 			"Name": "AudioModulationEditor",
