Checkpoint - WaveTable Integration w/Modulation & MetaSounds
- Removed duplicate implementation of ModPatch editor stuff now that WT Plugin covers it
- Add UWaveTableBank asset & WaveTable MetaSound DataType
- Add Full Sin as curve type
- Fix-up ability to make WaveTable assets bipolar but leave Modulation Patches unipolar
- Inherit ModulationTransform from new WaveTableTransform for extensibility to WT implementation
- Move ModulationCurve --> WaveTableCurve
- FastSin to Sin implementation in now WaveTableCurve (no reason to use fast implementation when using WT or interpolating curve key implementation at block rate)

diff --git a/Engine/Plugins/Runtime/AudioModulation/AudioModulation.uplugin b/Engine/Plugins/Runtime/AudioModulation/AudioModulation.uplugin
--- a/Engine/Plugins/Runtime/AudioModulation/AudioModulation.uplugin
+++ b/Engine/Plugins/Runtime/AudioModulation/AudioModulation.uplugin
@@ -32,5 +32,9 @@
 			"Name": "MetaSound",
 			"Enabled": true
 		},
+		{
+			"Name": "WaveTable",
+			"Enabled": true
+		}
 	]	
 }
\ No newline at end of file
diff --git a/Engine/Plugins/Runtime/AudioModulation/Config/BaseAudioModulation.ini b/Engine/Plugins/Runtime/AudioModulation/Config/BaseAudioModulation.ini
--- a/Engine/Plugins/Runtime/AudioModulation/Config/BaseAudioModulation.ini
+++ b/Engine/Plugins/Runtime/AudioModulation/Config/BaseAudioModulation.ini
@@ -1,5 +1,6 @@
 [CoreRedirects]
 +EnumRedirects=(OldName="FSoundModulationLFOShape",NewName="ESoundModulationLFOShape")
++EnumRedirects=(OldName="ESoundModulationCurve",NewName="EWaveTableCurve")
 
 [/Script/AudioModulation.AudioModulationSettings]
 +Parameters=/AudioModulation/BitDepth.BitDepth
diff --git a/Engine/Plugins/Runtime/AudioModulation/Source/AudioModulation/AudioModulation.Build.cs b/Engine/Plugins/Runtime/AudioModulation/Source/AudioModulation/AudioModulation.Build.cs
--- a/Engine/Plugins/Runtime/AudioModulation/Source/AudioModulation/AudioModulation.Build.cs
+++ b/Engine/Plugins/Runtime/AudioModulation/Source/AudioModulation/AudioModulation.Build.cs
@@ -27,7 +27,8 @@ namespace UnrealBuildTool.Rules
 					"CoreUObject",
 					"DeveloperSettings",
 					"Engine",
-					"SignalProcessing"
+					"SignalProcessing",
+					"WaveTable"
 				}
 			);
 
diff --git a/Engine/Plugins/Runtime/AudioModulation/Source/AudioModulation/Private/SoundModulationPatch.cpp b/Engine/Plugins/Runtime/AudioModulation/Source/AudioModulation/Private/SoundModulationPatch.cpp
--- a/Engine/Plugins/Runtime/AudioModulation/Source/AudioModulation/Private/SoundModulationPatch.cpp
+++ b/Engine/Plugins/Runtime/AudioModulation/Source/AudioModulation/Private/SoundModulationPatch.cpp
@@ -9,7 +9,6 @@
 #include "SoundControlBus.h"
 #include "SoundModulationParameter.h"
 #include "SoundModulationPatchProxy.h"
-#include "SoundModulationTransform.h"
 
 
 #define LOCTEXT_NAMESPACE "SoundModulationPatch"
@@ -41,6 +40,14 @@ const Audio::FModulationParameter& USoundModulationPatch::GetOutputParameter() c
 #if WITH_EDITOR
 void USoundModulationPatch::PostEditChangeProperty(FPropertyChangedEvent& InPropertyChangedEvent)
 {
+	if (InPropertyChangedEvent.Property && InPropertyChangedEvent.ChangeType != EPropertyChangeType::Interactive)
+	{
+		for (FSoundControlModulationInput& Input : PatchSettings.Inputs)
+		{
+			Input.Transform.CacheWaveTable(PatchSettings.WaveTableResolution, false /* bBipolar */);
+		}
+	}
+
 	AudioModulation::IterateModulationManagers([this](AudioModulation::FAudioModulationManager& OutModulation)
 	{
 		OutModulation.UpdateModulator(*this);
@@ -51,6 +58,14 @@ void USoundModulationPatch::PostEditChangeProperty(FPropertyChangedEvent& InProp
 
 void USoundModulationPatch::PostEditChangeChainProperty(FPropertyChangedChainEvent& InPropertyChangedEvent)
 {
+	if (InPropertyChangedEvent.Property && InPropertyChangedEvent.ChangeType != EPropertyChangeType::Interactive)
+	{
+		for (FSoundControlModulationInput& Input : PatchSettings.Inputs)
+		{
+			Input.Transform.CacheWaveTable(PatchSettings.WaveTableResolution, false /* bBipolar */);
+		}
+	}
+
 	AudioModulation::IterateModulationManagers([this](AudioModulation::FAudioModulationManager& OutModulation)
 	{
 		OutModulation.UpdateModulator(*this);
@@ -58,6 +73,19 @@ void USoundModulationPatch::PostEditChangeChainProperty(FPropertyChangedChainEve
 
 	Super::PostEditChangeChainProperty(InPropertyChangedEvent);
 }
+
+void USoundModulationPatch::PreSave(FObjectPreSaveContext InSaveContext)
+{
+	if (!InSaveContext.IsCooking())
+	{
+		for (FSoundControlModulationInput& Input : PatchSettings.Inputs)
+		{
+			Input.Transform.CacheWaveTable(PatchSettings.WaveTableResolution, false /* bBipolar */);
+		}
+	}
+
+	Super::PreSave(InSaveContext);
+}
 #endif // WITH_EDITOR
 
 FSoundControlModulationInput::FSoundControlModulationInput()
diff --git a/Engine/Plugins/Runtime/AudioModulation/Source/AudioModulation/Private/SoundModulationPatchProxy.cpp b/Engine/Plugins/Runtime/AudioModulation/Source/AudioModulation/Private/SoundModulationPatchProxy.cpp
--- a/Engine/Plugins/Runtime/AudioModulation/Source/AudioModulation/Private/SoundModulationPatchProxy.cpp
+++ b/Engine/Plugins/Runtime/AudioModulation/Source/AudioModulation/Private/SoundModulationPatchProxy.cpp
@@ -10,7 +10,6 @@
 #include "SoundControlBusProxy.h"
 #include "SoundModulationParameter.h"
 #include "SoundModulationProxy.h"
-#include "SoundModulationTransform.h"
 
 
 namespace AudioModulation
diff --git a/Engine/Plugins/Runtime/AudioModulation/Source/AudioModulation/Private/SoundModulationPatchProxy.h b/Engine/Plugins/Runtime/AudioModulation/Source/AudioModulation/Private/SoundModulationPatchProxy.h
--- a/Engine/Plugins/Runtime/AudioModulation/Source/AudioModulation/Private/SoundModulationPatchProxy.h
+++ b/Engine/Plugins/Runtime/AudioModulation/Source/AudioModulation/Private/SoundModulationPatchProxy.h
@@ -1,10 +1,12 @@
 // Copyright Epic Games, Inc. All Rights Reserved.
 #pragma once
 
+#include "Algo/Transform.h"
 #include "AudioDeviceManager.h"
 #include "IAudioModulation.h"
 #include "SoundControlBusProxy.h"
 #include "SoundModulationParameter.h"
+#include "SoundModulationPatch.h"
 #include "SoundModulationProxy.h"
 #include "Templates/Function.h"
 
@@ -87,9 +89,14 @@ namespace AudioModulation
 		{
 			for (const FSoundControlModulationInput& Input : InPatch.PatchSettings.Inputs)
 			{
-				if (Input.GetBus())
+				if (Input.Bus)
 				{
-					InputSettings.Add(FModulationInputSettings(Input));
+					FModulationInputSettings NewInputSettings(Input);
+
+					// Required to avoid referencing external UObject from settings on non-audio/game threads
+					NewInputSettings.Transform.CacheCurve();
+
+					InputSettings.Add(MoveTemp(NewInputSettings));
 				}
 			}
 		}
diff --git a/Engine/Plugins/Runtime/AudioModulation/Source/AudioModulation/Private/SoundModulationTransform.cpp b/Engine/Plugins/Runtime/AudioModulation/Source/AudioModulation/Private/SoundModulationTransform.cpp
deleted file mode 100644
diff --git a/Engine/Plugins/Runtime/AudioModulation/Source/AudioModulation/Public/SoundModulationPatch.h b/Engine/Plugins/Runtime/AudioModulation/Source/AudioModulation/Public/SoundModulationPatch.h
--- a/Engine/Plugins/Runtime/AudioModulation/Source/AudioModulation/Public/SoundModulationPatch.h
+++ b/Engine/Plugins/Runtime/AudioModulation/Source/AudioModulation/Public/SoundModulationPatch.h
@@ -1,12 +1,12 @@
 // Copyright Epic Games, Inc. All Rights Reserved.
 #pragma once
 
-#include "CoreMinimal.h"
 #include "IAudioModulation.h"
 #include "SoundModulationParameter.h"
-#include "SoundModulationTransform.h"
-#include "UObject/ObjectMacros.h"
 #include "UObject/Object.h"
+#include "UObject/ObjectMacros.h"
+#include "UObject/ObjectSaveContext.h"
+#include "WaveTableTransform.h"
 
 #include "SoundModulationPatch.generated.h"
 
@@ -14,6 +14,13 @@
 class USoundControlBus;
 
 
+USTRUCT(BlueprintType)
+struct FSoundModulationTransform : public FWaveTableTransform
+{
+	GENERATED_USTRUCT_BODY()
+};
+
+
 USTRUCT(BlueprintType)
 struct AUDIOMODULATION_API FSoundControlModulationInput
 {
@@ -53,6 +60,8 @@ struct AUDIOMODULATION_API FSoundControlModulationPatch
 	UPROPERTY(EditAnywhere, BlueprintReadWrite, Category = Inputs)
 	TArray<FSoundControlModulationInput> Inputs;
 
+	UPROPERTY(EditAnywhere, BlueprintReadWrite, Category = Inputs)
+	EWaveTableResolution WaveTableResolution;
 };
 
 UCLASS(config = Engine, editinlinenew, BlueprintType)
@@ -70,7 +79,9 @@ public:
 
 	virtual TUniquePtr<Audio::IModulatorSettings> CreateProxySettings() const override;
 
+
 #if WITH_EDITOR
+	virtual void PreSave(FObjectPreSaveContext InSaveContext) override;
 	virtual void PostEditChangeProperty(FPropertyChangedEvent& InPropertyChangedEvent) override;
 	virtual void PostEditChangeChainProperty(FPropertyChangedChainEvent& InPropertyChangedEvent) override;
 #endif // WITH_EDITOR
diff --git a/Engine/Plugins/Runtime/AudioModulation/Source/AudioModulation/Public/SoundModulationTransform.h b/Engine/Plugins/Runtime/AudioModulation/Source/AudioModulation/Public/SoundModulationTransform.h
deleted file mode 100644
diff --git a/Engine/Plugins/Runtime/AudioModulation/Source/AudioModulationEditor/AudioModulationEditor.build.cs b/Engine/Plugins/Runtime/AudioModulation/Source/AudioModulationEditor/AudioModulationEditor.build.cs
--- a/Engine/Plugins/Runtime/AudioModulation/Source/AudioModulationEditor/AudioModulationEditor.build.cs
+++ b/Engine/Plugins/Runtime/AudioModulation/Source/AudioModulationEditor/AudioModulationEditor.build.cs
@@ -32,6 +32,8 @@ public class AudioModulationEditor : ModuleRules
 				"CurveEditor",
 				"EditorWidgets",
 				"ToolWidgets",
+				"WaveTable",
+				"WaveTableEditor"
 			}
 		);
 
diff --git a/Engine/Plugins/Runtime/AudioModulation/Source/AudioModulationEditor/Private/AudioModulationEditor.cpp b/Engine/Plugins/Runtime/AudioModulation/Source/AudioModulationEditor/Private/AudioModulationEditor.cpp
--- a/Engine/Plugins/Runtime/AudioModulation/Source/AudioModulationEditor/Private/AudioModulationEditor.cpp
+++ b/Engine/Plugins/Runtime/AudioModulation/Source/AudioModulationEditor/Private/AudioModulationEditor.cpp
@@ -24,7 +24,7 @@
 #include "Modules/ModuleManager.h"
 #include "Sound/SoundBase.h"
 #include "SoundModulationParameter.h"
-#include "SoundModulationTransform.h"
+#include "SoundModulationPatch.h"
 #include "Templates/SharedPointer.h"
 #include "Textures/SlateIcon.h"
 #include "UObject/UObjectIterator.h"
@@ -109,7 +109,7 @@ void FAudioModulationEditorModule::StartupModule()
 	RegisterCustomPropertyLayouts();
 
 	ICurveEditorModule& CurveEditorModule = FModuleManager::LoadModuleChecked<ICurveEditorModule>("CurveEditor");
-	FModPatchCurveEditorModel::ViewId = CurveEditorModule.RegisterView(FOnCreateCurveEditorView::CreateStatic(
+	FModPatchCurveEditorModel::ModPatchViewId = CurveEditorModule.RegisterView(FOnCreateCurveEditorView::CreateStatic(
 		[](TWeakPtr<FCurveEditor> WeakCurveEditor) -> TSharedRef<SCurveEditorView>
 		{
 			return SNew(SModulationPatchEditorViewStacked, WeakCurveEditor);
@@ -187,9 +187,10 @@ void FAudioModulationEditorModule::ShutdownModule()
 
 	if (ICurveEditorModule* CurveEditorModule = FModuleManager::GetModulePtr<ICurveEditorModule>("CurveEditor"))
 	{
-		CurveEditorModule->UnregisterView(FModPatchCurveEditorModel::ViewId);
+		CurveEditorModule->UnregisterView(FModPatchCurveEditorModel::ModPatchViewId);
 	}
-	FModPatchCurveEditorModel::ViewId = ECurveEditorViewID::Invalid;
+
+	FModPatchCurveEditorModel::ModPatchViewId = ECurveEditorViewID::Invalid;
 
 	FSlateStyleRegistry::UnRegisterSlateStyle(*StyleSet.Get());
 }
diff --git a/Engine/Plugins/Runtime/AudioModulation/Source/AudioModulationEditor/Private/Editors/ModulationPatchCurveEditorViewStacked.cpp b/Engine/Plugins/Runtime/AudioModulation/Source/AudioModulationEditor/Private/Editors/ModulationPatchCurveEditorViewStacked.cpp
--- a/Engine/Plugins/Runtime/AudioModulation/Source/AudioModulationEditor/Private/Editors/ModulationPatchCurveEditorViewStacked.cpp
+++ b/Engine/Plugins/Runtime/AudioModulation/Source/AudioModulationEditor/Private/Editors/ModulationPatchCurveEditorViewStacked.cpp
@@ -22,6 +22,8 @@
 
 namespace PatchCurveViewUtils
 {
+	static const FText NormalizedAxisName = LOCTEXT("ModulationCurveDisplayTitle_Normalized", "Normalized");
+
 	void FormatLabel(const USoundModulationParameter& InParameter, const FNumberFormattingOptions& InNumFormatOptions, FText& InOutLabel)
 	{
 		const float NormalizedValue = FCString::Atof(*InOutLabel.ToString());
@@ -31,118 +33,43 @@ namespace PatchCurveViewUtils
 	}
 } // namespace PatchCurveViewUtils
 
-
-FModPatchCurveEditorModel::FModPatchCurveEditorModel(FRichCurve& InRichCurve, UObject* InOwner, EModPatchOutputEditorCurveSource InSource, int32 InInputIndex)
-	: FRichCurveEditorModelRaw(&InRichCurve, InOwner)
-	, Patch(CastChecked<USoundModulationPatch>(InOwner))
-	, InputIndex(InInputIndex)
-	, Source(InSource)
-{
-	check(InOwner);
-	Refresh(InSource, InInputIndex);
-}
-
-void FModPatchCurveEditorModel::Refresh(EModPatchOutputEditorCurveSource InSource, int32 InInputIndex)
+void SModulationPatchEditorViewStacked::FormatInputLabel(const WaveTable::Editor::FWaveTableCurveModel& EditorModel, const FNumberFormattingOptions& InLabelFormat, FText& InOutLabel) const
 {
-	InputAxisName = LOCTEXT("ModulationCurveDisplayTitle_Normalized", "Normalized");
-	FText OutputAxisName = InputAxisName;
-
-	FSoundControlModulationInput* Input = nullptr;
-	Source = EModPatchOutputEditorCurveSource::Unset;
-	InputIndex = -1;
-
-	if (Patch.IsValid() && InInputIndex >= 0 && InInputIndex < Patch->PatchSettings.Inputs.Num())
+	if (const USoundModulationPatch* Patch = static_cast<const FModPatchCurveEditorModel&>(EditorModel).GetPatch())
 	{
-		InputIndex = InInputIndex;
-		Source = InSource;
-
-		static const FText AxisNameFormat = LOCTEXT("ModulationCurveDisplayTitle_AxisNameFormat", "{0} ({1})");
-
-		Input = &Patch->PatchSettings.Inputs[InInputIndex];
-		if (const USoundControlBus* Bus = Input->GetBus())
+		const int32 Index = EditorModel.GetCurveIndex();
+		const FSoundControlModulationInput& Input = Patch->PatchSettings.Inputs[Index];
+		if (const USoundControlBus* Bus = Input.Bus)
 		{
-			if (USoundModulationParameter* Parameter = Bus->Parameter)
+			if (const USoundModulationParameter* Parameter = Bus->Parameter)
 			{
-				InputAxisName = FText::Format(AxisNameFormat, FText::FromString(Parameter->GetName()), Parameter->Settings.UnitDisplayName);
+				PatchCurveViewUtils::FormatLabel(*Parameter, InLabelFormat, InOutLabel);
 			}
 		}
-
-		if (USoundModulationParameter* Parameter = Patch->PatchSettings.OutputParameter)
-		{
-			OutputAxisName = FText::Format(AxisNameFormat, FText::FromString(Parameter->GetName()), Parameter->Settings.UnitDisplayName);
-		}
 	}
+}
 
-	AxesDescriptor = FText::Format(LOCTEXT("ModulationCurveDisplayTitle_AxesFormat", "X = {0}, Y = {1}"), InputAxisName, OutputAxisName);
-
-	bKeyDrawEnabled = true;
-	Color = UAudioModulationStyle::GetControlBusColor();
-	IntentionName = TEXT("AudioControlValue");
-	SupportedViews = ViewId;
-
-	ShortDisplayName = LOCTEXT("ModulationCurveDisplayTitle_BusUnset", "Bus (Unset)");
-	if (ensure(Input))
-	{
-		if (const USoundControlBus* Bus = Input->GetBus())
-		{
-			ShortDisplayName = FText::FromString(Input->GetBus()->GetName());
-		}
-	}
-
-	bKeyDrawEnabled = false;
-
-	const bool bIsBypassed = GetIsBypassed();
-	if (bIsBypassed)
-	{
-		LongDisplayName = FText::Format(LOCTEXT("ModulationCurveDisplay_BypassedNameFormat", "{0} (Bypassed)"), ShortDisplayName);
-	}
-	else
+void SModulationPatchEditorViewStacked::FormatOutputLabel(const WaveTable::Editor::FWaveTableCurveModel& EditorModel, const FNumberFormattingOptions& InLabelFormat, FText& InOutLabel) const
+{
+	if (const USoundModulationPatch* Patch = static_cast<const FModPatchCurveEditorModel&>(EditorModel).GetPatch())
 	{
-		UEnum* Enum = StaticEnum<EModPatchOutputEditorCurveSource>();
-		check(Enum);
-
-		FString EnumValueName = Enum->GetValueAsString(Source);
-		int32 DelimIndex = -1;
-		EnumValueName.FindLastChar(':', DelimIndex);
-		if (DelimIndex > 0 && DelimIndex < EnumValueName.Len() - 1)
+		if (const USoundModulationParameter* Parameter = Patch->PatchSettings.OutputParameter)
 		{
-			EnumValueName.RightChopInline(DelimIndex + 1);
+			PatchCurveViewUtils::FormatLabel(*Parameter, InLabelFormat, InOutLabel);
 		}
-
-		FText CurveSourceText = FText::FromString(EnumValueName);
-
-		if (Source == EModPatchOutputEditorCurveSource::Custom)
-		{
-			bKeyDrawEnabled = true;
-		}
-		else if (Source == EModPatchOutputEditorCurveSource::Shared)
-		{
-			check(Input);
-			UCurveFloat* SharedCurve = Input->Transform.CurveShared;
-			check(SharedCurve);
-			CurveSourceText = FText::Format(LOCTEXT("ModulationCurveDisplayTitle_SharedNameFormat", "Shared, {0}"), FText::FromString(SharedCurve->GetName()));
-			bKeyDrawEnabled = true;
-		}
-
-		LongDisplayName = FText::Format(LOCTEXT("ModulationCurveDisplayTitle_NameFormat", "{0}: {1} ({2})"), FText::AsNumber(InInputIndex), ShortDisplayName, CurveSourceText);
 	}
 }
 
-FLinearColor FModPatchCurveEditorModel::GetColor() const
-{
-	return !GetIsBypassed() && (Source == EModPatchOutputEditorCurveSource::Custom)
-		? Color
-		: Color.Desaturate(0.35f);
-}
+ECurveEditorViewID FModPatchCurveEditorModel::ModPatchViewId = ECurveEditorViewID::Invalid;
 
-const FText& FModPatchCurveEditorModel::GetAxesDescriptor() const
+FModPatchCurveEditorModel::FModPatchCurveEditorModel(FRichCurve& InRichCurve, UObject* InOwner, EWaveTableCurveSource InSource)
+	: WaveTable::Editor::FWaveTableCurveModel(InRichCurve, InOwner, InSource)
 {
-	return AxesDescriptor;
 }
 
 bool FModPatchCurveEditorModel::GetIsBypassed() const
 {
-	if (Patch.IsValid())
+	if (const USoundModulationPatch* Patch = GetPatch())
 	{
 		return Patch->PatchSettings.bBypass;
 	}
@@ -150,19 +77,11 @@ bool FModPatchCurveEditorModel::GetIsBypassed() const
 	return true;
 }
 
-EModPatchOutputEditorCurveSource FModPatchCurveEditorModel::GetSource() const
+USoundModulationPatch* FModPatchCurveEditorModel::GetPatch()
 {
-	return Source;
-}
-
-USoundModulationParameter* FModPatchCurveEditorModel::GetPatchInputParameter() const
-{
-	if (Patch.IsValid() && InputIndex >= 0 && InputIndex < Patch->PatchSettings.Inputs.Num())
+	if (ParentObject.IsValid())
 	{
-		if (const USoundControlBus* Bus = Patch->PatchSettings.Inputs[InputIndex].GetBus())
-		{
-			return Bus->Parameter;
-		}
+		return CastChecked<USoundModulationPatch>(ParentObject);
 	}
 
 	return nullptr;
@@ -170,295 +89,51 @@ USoundModulationParameter* FModPatchCurveEditorModel::GetPatchInputParameter() c
 
 const USoundModulationPatch* FModPatchCurveEditorModel::GetPatch() const
 {
-	return Patch.Get();
-}
-
-bool FModPatchCurveEditorModel::IsReadOnly() const
-{
-	return Source != EModPatchOutputEditorCurveSource::Custom;
-}
-
-ECurveEditorViewID FModPatchCurveEditorModel::ViewId = ECurveEditorViewID::Invalid;
-
-void SModulationPatchEditorViewStacked::Construct(const FArguments& InArgs, TWeakPtr<FCurveEditor> InCurveEditor)
-{
-	SCurveEditorViewStacked::Construct(InArgs, InCurveEditor);
-
-	StackedHeight = 300.0f;
-	StackedPadding = 60.0f;
-}
-
-void SModulationPatchEditorViewStacked::PaintView(const FPaintArgs& Args, const FGeometry& AllottedGeometry, const FSlateRect& MyCullingRect, FSlateWindowElementList& OutDrawElements, int32 BaseLayerId, const FWidgetStyle& InWidgetStyle, bool bParentEnabled) const
-{
-	TSharedPtr<FCurveEditor> CurveEditor = WeakCurveEditor.Pin();
-	if (CurveEditor)
+	if (ParentObject.IsValid())
 	{
-		const ESlateDrawEffect DrawEffects = ShouldBeEnabled(bParentEnabled) ? ESlateDrawEffect::None : ESlateDrawEffect::DisabledEffect;
-
-		DrawBackground(AllottedGeometry, OutDrawElements, BaseLayerId, DrawEffects);
-		DrawViewGrids(AllottedGeometry, MyCullingRect, OutDrawElements, BaseLayerId, DrawEffects);
-		DrawLabels(AllottedGeometry, MyCullingRect, OutDrawElements, BaseLayerId, DrawEffects);
-		DrawCurves(CurveEditor.ToSharedRef(), AllottedGeometry, MyCullingRect, OutDrawElements, BaseLayerId, InWidgetStyle, DrawEffects);
+		return CastChecked<const USoundModulationPatch>(ParentObject);
 	}
+
+	return nullptr;
 }
 
-void SModulationPatchEditorViewStacked::DrawLabels(const FGeometry& AllottedGeometry, const FSlateRect& MyCullingRect, FSlateWindowElementList& OutDrawElements, int32 BaseLayerId, ESlateDrawEffect DrawEffects) const
+FWaveTableTransform* FModPatchCurveEditorModel::GetTransform()
 {
-	TSharedPtr<FCurveEditor> CurveEditor = WeakCurveEditor.Pin();
-	if (!CurveEditor)
+	if (USoundModulationPatch* Patch = GetPatch())
 	{
-		return;
+		FSoundControlModulationInput& Input = Patch->PatchSettings.Inputs[GetCurveIndex()];
+		return &Input.Transform;
 	}
 
-	const int32 LabelLayerId = BaseLayerId + CurveViewConstants::ELayerOffset::Labels;
-
-	const double ValuePerPixel = 1.0 / StackedHeight;
-	const double ValueSpacePadding = StackedPadding * ValuePerPixel;
-
-	const FSlateFontInfo FontInfo = FAppStyle::GetFontStyle("CurveEd.LabelFont");
-	const FVector2D LocalSize = AllottedGeometry.GetLocalSize();
-	const FCurveEditorScreenSpace ViewSpace = GetViewSpace();
-
-	// Draw the curve labels for each view
-	for (auto It = CurveInfoByID.CreateConstIterator(); It; ++It)
-	{
-		FCurveModel* Curve = CurveEditor->FindCurve(It.Key());
-		if (!ensureAlways(Curve))
-		{
-			continue;
-		}
-
-		const int32  CurveIndexFromBottom = CurveInfoByID.Num() - It->Value.CurveIndex - 1;
-		const double PaddingToBottomOfView = (CurveIndexFromBottom + 1) * ValueSpacePadding;
-
-		const float PixelBottom = ViewSpace.ValueToScreen(CurveIndexFromBottom + PaddingToBottomOfView);
-		const float PixelTop = ViewSpace.ValueToScreen(CurveIndexFromBottom + PaddingToBottomOfView + 1.0);
-
-		if (!FSlateRect::DoRectanglesIntersect(MyCullingRect, TransformRect(AllottedGeometry.GetAccumulatedLayoutTransform(), FSlateRect(0, PixelTop, LocalSize.X, PixelBottom))))
-		{
-			continue;
-		}
-
-		const TSharedRef<FSlateFontMeasure> FontMeasure = FSlateApplication::Get().GetRenderer()->GetFontMeasureService();
-
-		// Render label
-		FText Label = Curve->GetLongDisplayName();
-
-		static const float LabelOffsetX = 15.0f;
-		static const float LabelOffsetY = 35.0f;
-		FVector2D LabelPosition(LabelOffsetX, PixelTop - LabelOffsetY);
-		FPaintGeometry LabelGeometry = AllottedGeometry.ToPaintGeometry(FSlateLayoutTransform(LabelPosition));
-		const FVector2D LabelSize = FontMeasure->Measure(Label, FontInfo);
-
-		FSlateDrawElement::MakeText(OutDrawElements, LabelLayerId + 1, LabelGeometry, Label, FontInfo, DrawEffects, Curve->GetColor());
-
-		// Render axes descriptor
-		FText Descriptor = static_cast<FModPatchCurveEditorModel*>(Curve)->GetAxesDescriptor();
-
-		const FVector2D DescriptorSize = FontMeasure->Measure(Descriptor, FontInfo);
-
-		static const float LabelBufferX = 20.0f; // Keeps label and axes descriptor visually separated
-		static const float GutterBufferX = 20.0f; // Accounts for potential scroll bar
-		const float ViewWidth = ViewSpace.GetPhysicalWidth();
-		const float FloatingDescriptorX = ViewWidth - DescriptorSize.X;
-		LabelPosition = FVector2D(FMath::Max(LabelSize.X + LabelBufferX, FloatingDescriptorX - GutterBufferX), PixelTop - LabelOffsetY);
-		LabelGeometry = AllottedGeometry.ToPaintGeometry(FSlateLayoutTransform(LabelPosition));
-
-		FSlateDrawElement::MakeText(OutDrawElements, LabelLayerId + 1, LabelGeometry, Descriptor, FontInfo, DrawEffects, Curve->GetColor());
-	}
+	return nullptr;
 }
 
-void SModulationPatchEditorViewStacked::DrawViewGrids(const FGeometry& AllottedGeometry, const FSlateRect& MyCullingRect, FSlateWindowElementList& OutDrawElements, int32 BaseLayerId, ESlateDrawEffect DrawEffects) const
+void FModPatchCurveEditorModel::RefreshCurveDescriptorText(const FWaveTableTransform& InTransform, FText& OutShortDisplayName, FText& OutInputAxisName, FText& OutOutputAxisName)
 {
-	TSharedPtr<FCurveEditor> CurveEditor = WeakCurveEditor.Pin();
-	if (!CurveEditor)
-	{
-		return;
-	}
-
-	// Rendering info
-	const float          Width = AllottedGeometry.GetLocalSize().X;
-	const float          Height = AllottedGeometry.GetLocalSize().Y;
-	const FSlateBrush*   WhiteBrush = FAppStyle::GetBrush("WhiteBrush");
+	OutShortDisplayName = LOCTEXT("ModulationCurveDisplayTitle_BusUnset", "Bus (Unset)");
+	OutInputAxisName = PatchCurveViewUtils::NormalizedAxisName;
+	OutOutputAxisName = PatchCurveViewUtils::NormalizedAxisName;
 
-	FModPatchViewGridDrawInfo DrawInfo(&AllottedGeometry, GetViewSpace(), CurveEditor->GetPanel()->GetGridLineTint(), BaseLayerId);
-
-	TArray<float> MajorGridLinesX, MinorGridLinesX;
-	TArray<FText> MajorGridLabelsX;
-
-	GetGridLinesX(CurveEditor.ToSharedRef(), MajorGridLinesX, MinorGridLinesX, &MajorGridLabelsX);
-
-	const double ValuePerPixel = 1.0 / StackedHeight;
-	const double ValueSpacePadding = StackedPadding * ValuePerPixel;
-
-	for (auto It = CurveInfoByID.CreateConstIterator(); It; ++It)
+	if (const USoundModulationPatch* Patch = GetPatch())
 	{
-		DrawInfo.SetCurveModel(CurveEditor->FindCurve(It.Key()));
-		if (!ensureAlways(DrawInfo.GetCurveModel()))
-		{
-			continue;
-		}
-
-		TArray<FText> CurveModelGridLabelsX = MajorGridLabelsX;
-		check(MajorGridLinesX.Num() == CurveModelGridLabelsX.Num());
-
-		if (const FModPatchCurveEditorModel* EditorModel = static_cast<const FModPatchCurveEditorModel*>(DrawInfo.GetCurveModel()))
-		{
-			if (USoundModulationParameter* InputParameter = EditorModel->GetPatchInputParameter())
-			{
-				for (int32 i = 0; i < CurveModelGridLabelsX.Num(); ++i)
-				{
-					FText& Label = CurveModelGridLabelsX[i];
-					PatchCurveViewUtils::FormatLabel(*InputParameter, DrawInfo.LabelFormat, Label);
-				}
-			}
-		}
-
-		const int32 Index = CurveInfoByID.Num() - It->Value.CurveIndex - 1;
-		double Padding = (Index + 1) * ValueSpacePadding;
-		DrawInfo.SetLowerValue(Index + Padding);
+		const int32 Index = GetCurveIndex();
 
-		FSlateRect BoundsRect(0, DrawInfo.GetPixelTop(), Width, DrawInfo.GetPixelBottom());
-		if (!FSlateRect::DoRectanglesIntersect(MyCullingRect, TransformRect(AllottedGeometry.GetAccumulatedLayoutTransform(), BoundsRect)))
+		const FSoundControlModulationInput& Input = Patch->PatchSettings.Inputs[Index];
+		if (const USoundControlBus* Bus = Input.GetBus())
 		{
-			continue;
-		}
-
-		// Tint the views based on their curve color
-		{
-			FLinearColor CurveColorTint = DrawInfo.GetCurveModel()->GetColor().CopyWithNewOpacity(0.05f);
-			const FPaintGeometry BoxGeometry = AllottedGeometry.ToPaintGeometry(
-				FVector2D(Width, StackedHeight),
-				FSlateLayoutTransform(FVector2D(0.f, DrawInfo.GetPixelTop()))
-			);
-
-			const int32 GridOverlayLayerId = DrawInfo.GetBaseLayerId() + CurveViewConstants::ELayerOffset::GridOverlays;
-			FSlateDrawElement::MakeBox(OutDrawElements, GridOverlayLayerId, BoxGeometry, WhiteBrush, DrawEffects, CurveColorTint);
-		}
-
-		// Horizontal grid lines
-		DrawInfo.LinePoints[0].X = 0.0;
-		DrawInfo.LinePoints[1].X = Width;
-
-		DrawViewGridLineX(OutDrawElements, DrawInfo, DrawEffects, 1.0  /* OffsetAlpha */, true  /* bIsMajor */);
-		DrawViewGridLineX(OutDrawElements, DrawInfo, DrawEffects, 0.75 /* OffsetAlpha */, false /* bIsMajor */);
-		DrawViewGridLineX(OutDrawElements, DrawInfo, DrawEffects, 0.5  /* OffsetAlpha */, true  /* bIsMajor */);
-		DrawViewGridLineX(OutDrawElements, DrawInfo, DrawEffects, 0.25 /* OffsetAlpha */, false /* bIsMajor */);
-		DrawViewGridLineX(OutDrawElements, DrawInfo, DrawEffects, 0.0  /* OffsetAlpha */, true  /* bIsMajor */);
-
-		// Vertical grid lines
-		{
-			DrawInfo.LinePoints[0].Y = DrawInfo.GetPixelTop();
-			DrawInfo.LinePoints[1].Y = DrawInfo.GetPixelBottom();
-
-			// Draw major vertical grid lines
-			for (int i = 0; i < MajorGridLinesX.Num(); ++i)
-			{
-				const float VerticalLine = FMath::RoundToFloat(MajorGridLinesX[i]);
-				const FText Label = CurveModelGridLabelsX[i];
-
-				DrawViewGridLineY(VerticalLine, OutDrawElements, DrawInfo, DrawEffects, &Label, true /* bIsMajor */);
-			}
-
-			// Now draw the minor vertical lines which are drawn with a lighter color.
-			for (float VerticalLine : MinorGridLinesX)
-			{
-				VerticalLine = FMath::RoundToFloat(VerticalLine);
-				DrawViewGridLineY(VerticalLine, OutDrawElements, DrawInfo, DrawEffects, nullptr /* Label */, false /* bIsMajor */);
-			}
-		}
-	}
-}
-
-void SModulationPatchEditorViewStacked::DrawViewGridLineX(FSlateWindowElementList& OutDrawElements, FModPatchViewGridDrawInfo& DrawInfo, ESlateDrawEffect DrawEffects, double OffsetAlpha, bool bIsMajor) const
-{
-	double ValueMin;
-	double ValueMax;
-	DrawInfo.GetCurveModel()->GetValueRange(ValueMin, ValueMax);
+			OutShortDisplayName = FText::FromString(Bus->GetName());
 
-	const int32 GridLineLayerId = DrawInfo.GetBaseLayerId() + CurveViewConstants::ELayerOffset::GridLines;
-
-	const double LowerValue = DrawInfo.GetLowerValue();
-	const double PixelBottom = DrawInfo.GetPixelBottom();
-	const double PixelTop = DrawInfo.GetPixelTop();
-
-	const TSharedRef<FSlateFontMeasure> FontMeasure = FSlateApplication::Get().GetRenderer()->GetFontMeasureService();
-	const FSlateFontInfo FontInfo = FCoreStyle::Get().GetFontStyle("CurveEd.InfoFont");
-
-	FLinearColor Color = bIsMajor ? DrawInfo.GetMajorGridColor() : DrawInfo.GetMinorGridColor();
-
-	DrawInfo.LinePoints[0].Y = DrawInfo.LinePoints[1].Y = DrawInfo.ScreenSpace.ValueToScreen(LowerValue + OffsetAlpha);
-	FSlateDrawElement::MakeLines(OutDrawElements, GridLineLayerId, DrawInfo.PaintGeometry, DrawInfo.LinePoints, DrawEffects, Color, false);
-
-	FText Label = FText::AsNumber(FMath::Lerp(ValueMin, ValueMax, OffsetAlpha), &DrawInfo.LabelFormat);
-
-	if (const FModPatchCurveEditorModel* EditorModel = static_cast<const FModPatchCurveEditorModel*>(DrawInfo.GetCurveModel()))
-	{
-		if (const USoundModulationPatch* Patch = EditorModel->GetPatch())
-		{
-			if (USoundModulationParameter* OutputParameter = Patch->PatchSettings.OutputParameter)
+			if (const USoundModulationParameter* Parameter = Bus->Parameter)
 			{
-				PatchCurveViewUtils::FormatLabel(*OutputParameter, DrawInfo.LabelFormat, Label);
+				static const FText AxisNameFormat = LOCTEXT("ModulationCurveDisplayTitle_InputAxisNameFormat", "{0} ({1})");
+				OutInputAxisName = FText::Format(AxisNameFormat, FText::FromString(Parameter->GetName()), Parameter->Settings.UnitDisplayName);
 			}
 		}
-	}
-
-	const FVector2D LabelSize = FontMeasure->Measure(Label, FontInfo);
-	double LabelY = FMath::Lerp(PixelBottom, PixelTop, OffsetAlpha);
-	
-	// Offset label Y position below line only if the top gridline,
-	// otherwise push above
-	if (FMath::IsNearlyEqual(OffsetAlpha, 1.0))
-	{
-		LabelY += 5.0;
-	}
-	else
-	{
-		LabelY -= 15.0;
-	}
-
-	const FPaintGeometry LabelGeometry = DrawInfo.AllottedGeometry->ToPaintGeometry(FSlateLayoutTransform(FVector2D(CurveViewConstants::CurveLabelOffsetX, LabelY)));
-
-	FSlateDrawElement::MakeText(
-		OutDrawElements,
-		DrawInfo.GetBaseLayerId() + CurveViewConstants::ELayerOffset::GridLabels,
-		LabelGeometry,
-		Label,
-		FontInfo,
-		DrawEffects,
-		DrawInfo.GetLabelColor()
-	);
-}
 
-void SModulationPatchEditorViewStacked::DrawViewGridLineY(const float VerticalLine, FSlateWindowElementList& OutDrawElements, FModPatchViewGridDrawInfo& DrawInfo, ESlateDrawEffect DrawEffects, const FText* Label, bool bIsMajor) const
-{
-	const float Width = DrawInfo.AllottedGeometry->GetLocalSize().X;
-	if (VerticalLine >= 0 || VerticalLine <= FMath::RoundToFloat(Width))
-	{
-		const int32 GridLineLayerId = DrawInfo.GetBaseLayerId() + CurveViewConstants::ELayerOffset::GridLines;
-
-		const FLinearColor LineColor = bIsMajor ? DrawInfo.GetMajorGridColor() : DrawInfo.GetMinorGridColor();
-
-		DrawInfo.LinePoints[0].X = DrawInfo.LinePoints[1].X = VerticalLine;
-		FSlateDrawElement::MakeLines(OutDrawElements, GridLineLayerId, DrawInfo.PaintGeometry, DrawInfo.LinePoints, DrawEffects, LineColor, false);
-
-		if (Label)
+		if (USoundModulationParameter* Parameter = Patch->PatchSettings.OutputParameter)
 		{
-			const FSlateFontInfo FontInfo = FCoreStyle::Get().GetFontStyle("CurveEd.InfoFont");
-
-			const TSharedRef<FSlateFontMeasure> FontMeasure = FSlateApplication::Get().GetRenderer()->GetFontMeasureService();
-			const FVector2D LabelSize = FontMeasure->Measure(*Label, FontInfo);
-			const FPaintGeometry LabelGeometry = DrawInfo.AllottedGeometry->ToPaintGeometry(FSlateLayoutTransform(FVector2D(VerticalLine, DrawInfo.LinePoints[0].Y - LabelSize.Y * 1.2f)));
-
-			FSlateDrawElement::MakeText(
-				OutDrawElements,
-				DrawInfo.GetBaseLayerId() + CurveViewConstants::ELayerOffset::GridLabels,
-				LabelGeometry,
-				*Label,
-				FontInfo,
-				DrawEffects,
-				DrawInfo.GetLabelColor()
-			);
+			static const FText AxisNameFormat = LOCTEXT("ModulationCurveDisplayTitle_OutputAxisNameFormat", "{0} ({1})");
+			OutOutputAxisName = FText::Format(AxisNameFormat, FText::FromString(Parameter->GetName()), Parameter->Settings.UnitDisplayName);
 		}
 	}
 }
diff --git a/Engine/Plugins/Runtime/AudioModulation/Source/AudioModulationEditor/Private/Editors/ModulationPatchCurveEditorViewStacked.h b/Engine/Plugins/Runtime/AudioModulation/Source/AudioModulationEditor/Private/Editors/ModulationPatchCurveEditorViewStacked.h
--- a/Engine/Plugins/Runtime/AudioModulation/Source/AudioModulationEditor/Private/Editors/ModulationPatchCurveEditorViewStacked.h
+++ b/Engine/Plugins/Runtime/AudioModulation/Source/AudioModulationEditor/Private/Editors/ModulationPatchCurveEditorViewStacked.h
@@ -1,6 +1,7 @@
 // Copyright Epic Games, Inc. All Rights Reserved.
 #pragma once
 
+#include "AudioModulationStyle.h"
 #include "CurveModel.h"
 #include "Curves/RichCurve.h"
 #include "RichCurveEditorModel.h"
@@ -9,8 +10,8 @@
 #include "UObject/WeakObjectPtrTemplates.h"
 #include "Views/SCurveEditorViewStacked.h"
 #include "Views/SInteractiveCurveEditorView.h"
+#include "WaveTableCurveEditorViewStacked.h"
 
-#include "ModulationPatchCurveEditorViewStacked.generated.h"
 
 // Forward Declarations
 class UCurveFloat;
@@ -18,169 +19,33 @@ class USoundModulationParameter;
 class USoundModulationPatch;
 struct FSoundControlModulationInput;
 
+#define LOCTEXT_NAMESPACE "AudioModulation"
 
-UENUM()
-enum class EModPatchOutputEditorCurveSource : uint8
+class SModulationPatchEditorViewStacked : public WaveTable::Editor::SViewStacked
 {
-	Custom,
-	Expression,
-	Shared,
-	Unset
+	virtual void FormatInputLabel(const WaveTable::Editor::FWaveTableCurveModel& EditorModel, const FNumberFormattingOptions& InLabelFormat, FText& InOutLabel) const override;
+	virtual void FormatOutputLabel(const WaveTable::Editor::FWaveTableCurveModel& EditorModel, const FNumberFormattingOptions& InLabelFormat, FText& InOutLabel) const override;
 };
 
-
-struct FModPatchViewGridDrawInfo
+class FModPatchCurveEditorModel : public WaveTable::Editor::FWaveTableCurveModel
 {
-	const FGeometry* AllottedGeometry;
-
-	FPaintGeometry PaintGeometry;
-
-	ESlateDrawEffect DrawEffects;
-	FNumberFormattingOptions LabelFormat;
-
-	TArray<FVector2D> LinePoints;
-
-	FCurveEditorScreenSpace ScreenSpace;
-
-private:
-	FLinearColor MajorGridColor;
-	FLinearColor MinorGridColor;
-
-	int32 BaseLayerId;
-
-	const FCurveModel* CurveModel;
-
-	double LowerValue;
-	double PixelBottom;
-	double PixelTop;
-
 public:
-	FModPatchViewGridDrawInfo(const FGeometry* InAllottedGeometry, const FCurveEditorScreenSpace& InScreenSpace, FLinearColor InGridColor, int32 InBaseLayerId)
-		: AllottedGeometry(InAllottedGeometry)
-		, ScreenSpace(InScreenSpace)
-		, BaseLayerId(InBaseLayerId)
-		, CurveModel(nullptr)
-		, LowerValue(0)
-		, PixelBottom(0)
-		, PixelTop(0)
-	{
-		check(AllottedGeometry);
+	static ECurveEditorViewID ModPatchViewId;
 
-		// Pre-allocate an array of line points to draw our vertical lines. Each major grid line
-		// will overwrite the X value of both points but leave the Y value untouched so they draw from the bottom to the top.
-		LinePoints.Add(FVector2D(0.f, 0.f));
-		LinePoints.Add(FVector2D(0.f, 0.f));
+	FModPatchCurveEditorModel(FRichCurve& InRichCurve, UObject* InOwner, EWaveTableCurveSource InSource);
 
-		MajorGridColor = InGridColor;
-		MinorGridColor = InGridColor.CopyWithNewOpacity(InGridColor.A * 0.5f);
+	virtual void RefreshCurveDescriptorText(const FWaveTableTransform& InTransform, FText& OutShortDisplayName, FText& OutInputAxisName, FText& OutOutputAxisName) override;
 
-		PaintGeometry = InAllottedGeometry->ToPaintGeometry();
+	virtual FWaveTableTransform* GetTransform() override;
 
-		LabelFormat.SetMaximumFractionalDigits(2);
-	}
+	virtual ECurveEditorViewID GetViewId() const override { return ModPatchViewId; }
+	virtual FColor GetCurveColor() const override { return UAudioModulationStyle::GetControlBusColor(); }
+	virtual bool GetPropertyEditorDisabled() const override { return GetIsBypassed(); }
+	virtual FText GetPropertyEditorDisabledText() const override { return LOCTEXT("ModulationPatchCurveEditorView_Bypassed", "Bypassed"); }
 
-	void SetCurveModel(const FCurveModel * InCurveModel)
-	{
-		CurveModel = InCurveModel;
-	}
-
-	const FCurveModel* GetCurveModel() const
-	{
-		return CurveModel;
-	}
-
-	void SetLowerValue(double InLowerValue)
-	{
-		LowerValue = InLowerValue;
-		PixelBottom = ScreenSpace.ValueToScreen(InLowerValue);
-		PixelTop = ScreenSpace.ValueToScreen(InLowerValue + 1.0);
-	}
-
-	int32 GetBaseLayerId() const
-	{
-		return BaseLayerId;
-	}
-
-	FLinearColor GetLabelColor() const
-	{
-		check(CurveModel);
-		return GetCurveModel()->GetColor().CopyWithNewOpacity(0.7f);
-	}
-
-	double GetLowerValue() const
-	{
-		return LowerValue;
-	}
-
-	FLinearColor GetMajorGridColor() const
-	{
-		return MajorGridColor;
-	}
-
-	FLinearColor GetMinorGridColor() const
-	{
-		return MinorGridColor;
-	}
-
-	double GetPixelBottom() const
-	{
-		return PixelBottom;
-	}
-
-	double GetPixelTop() const
-	{
-		return PixelTop;
-	}
-};
-
-
-class FModPatchCurveEditorModel : public FRichCurveEditorModelRaw
-{
-public:
-	static ECurveEditorViewID ViewId;
-
-	FModPatchCurveEditorModel(FRichCurve& InRichCurve, UObject* InOwner, EModPatchOutputEditorCurveSource InSource, int32 InInputIndex);
-
-	const FText& GetAxesDescriptor() const;
 	bool GetIsBypassed() const;
-	USoundModulationParameter* GetPatchInputParameter() const;
-	const USoundModulationPatch* GetPatch() const;
-	EModPatchOutputEditorCurveSource GetSource() const;
-	void Refresh(EModPatchOutputEditorCurveSource InSource, int32 InInputIndex);
-
-	virtual bool IsReadOnly() const override;
-	virtual FLinearColor GetColor() const override;
-
-private:
-
-	TWeakObjectPtr<USoundModulationPatch> Patch;
-	int32 InputIndex = -1;
-	EModPatchOutputEditorCurveSource Source = EModPatchOutputEditorCurveSource::Unset;
 
-	FText InputAxisName;
-	FText AxesDescriptor;
-};
-
-class SModulationPatchEditorViewStacked : public SCurveEditorViewStacked
-{
-public:
-
-	void Construct(const FArguments& InArgs, TWeakPtr<FCurveEditor> InCurveEditor);
-
-protected:
-	virtual void PaintView(
-		const FPaintArgs&			Args,
-		const FGeometry&			AllottedGeometry,
-		const FSlateRect&			MyCullingRect,
-		FSlateWindowElementList&	OutDrawElements,
-		int32						BaseLayerId, 
-		const FWidgetStyle&			InWidgetStyle,
-		bool						bParentEnabled) const override;
-
-	virtual void DrawViewGrids(const FGeometry& AllottedGeometry, const FSlateRect& MyCullingRect, FSlateWindowElementList& OutDrawElements, int32 BaseLayerId, ESlateDrawEffect DrawEffects) const override;
-	virtual void DrawLabels(const FGeometry& AllottedGeometry, const FSlateRect& MyCullingRect, FSlateWindowElementList& OutDrawElements, int32 BaseLayerId, ESlateDrawEffect DrawEffects) const override;
-
-private:
-	void DrawViewGridLineX(FSlateWindowElementList& OutDrawElements, FModPatchViewGridDrawInfo& DrawInfo, ESlateDrawEffect DrawEffect, double OffsetAlpha, bool bIsMajor) const;
-	void DrawViewGridLineY(const float VerticalLine, FSlateWindowElementList& OutDrawElements, FModPatchViewGridDrawInfo &DrawInfo, ESlateDrawEffect DrawEffects, const FText* Label, bool bIsMajor) const;
+	USoundModulationPatch* GetPatch();
+	const USoundModulationPatch* GetPatch() const;
 };
+#undef LOCTEXT_NAMESPACE // "AudioModulation"
diff --git a/Engine/Plugins/Runtime/AudioModulation/Source/AudioModulationEditor/Private/Editors/ModulationPatchEditor.cpp b/Engine/Plugins/Runtime/AudioModulation/Source/AudioModulationEditor/Private/Editors/ModulationPatchEditor.cpp
--- a/Engine/Plugins/Runtime/AudioModulation/Source/AudioModulationEditor/Private/Editors/ModulationPatchEditor.cpp
+++ b/Engine/Plugins/Runtime/AudioModulation/Source/AudioModulationEditor/Private/Editors/ModulationPatchEditor.cpp
@@ -3,278 +3,36 @@
 
 #include "AudioModulationEditor.h"
 #include "AudioModulationStyle.h"
-#include "CommonFrameRates.h"
-#include "Containers/Set.h"
-#include "Curves/CurveBase.h"
-#include "Curves/CurveLinearColor.h"
-#include "CurveEditor.h"
-#include "CurveEditorCommands.h"
-#include "Styling/AppStyle.h"
-#include "Framework/Commands/UIAction.h"
-#include "Framework/MultiBox/MultiBoxDefs.h"
-#include "Framework/MultiBox/MultiBoxBuilder.h"
-#include "Framework/MultiBox/MultiBoxExtender.h"
-#include "IDetailsView.h"
-#include "Modules/ModuleManager.h"
-#include "PropertyEditorModule.h"
-#include "RichCurveEditorModel.h"
-#include "SCurveEditorPanel.h"
+#include "ModulationPatchCurveEditorViewStacked.h"
 #include "SoundModulationPatch.h"
-#include "SoundModulationTransform.h"
-#include "Tree/SCurveEditorTree.h"
-#include "Tree/ICurveEditorTreeItem.h"
-#include "Tree/SCurveEditorTreePin.h"
-#include "Widgets/Docking/SDockTab.h"
-#include "Widgets/Input/SButton.h"
-#include "Widgets/Input/SNumericDropDown.h"
-#include "Widgets/Layout/SBorder.h"
-#include "Widgets/SFrameRatePicker.h"
+#include "Templates/UniquePtr.h"
+#include "WaveTableCurveEditorViewStacked.h"
+#include "WaveTableSettings.h"
+#include "WaveTableSampler.h"
 
 
 #define LOCTEXT_NAMESPACE "ModulationPatchEditor"
 
 
-const FName FModulationPatchEditor::AppIdentifier(TEXT("ModulationPatchEditorApp"));
-const FName FModulationPatchEditor::CurveTabId(TEXT("ModulationPatchEditor_Curves"));
-const FName FModulationPatchEditor::PropertiesTabId(TEXT("ModulationPatchEditor_Properties"));
-
-
-FModulationPatchEditor::FModulationPatchEditor()
-{
-}
-
-void FModulationPatchEditor::RegisterTabSpawners(const TSharedRef<FTabManager>& InTabManager)
-{
-	WorkspaceMenuCategory = InTabManager->AddLocalWorkspaceMenuCategory(LOCTEXT("WorkspaceMenu_ModulationPatchEditor", "Modulation Patch Editor"));
-
-	FAssetEditorToolkit::RegisterTabSpawners(InTabManager);
-
-	InTabManager->RegisterTabSpawner(PropertiesTabId, FOnSpawnTab::CreateSP(this, &FModulationPatchEditor::SpawnTab_Properties))
-		.SetDisplayName(LOCTEXT("DetailsTab", "Details"))
-		.SetGroup(WorkspaceMenuCategory.ToSharedRef())
-		.SetIcon(FSlateIcon(FAppStyle::GetAppStyleSetName(), "LevelEditor.Tabs.Details"));
-
-	FSlateIcon CurveIcon(FAppStyle::GetAppStyleSetName(), "ModulationPatchEditor.Tabs.Properties");
-	InTabManager->RegisterTabSpawner(CurveTabId, FOnSpawnTab::CreateLambda([this](const FSpawnTabArgs& Args) { return SpawnTab_OutputCurve(Args); }))
-		.SetDisplayName(LOCTEXT("TransformCurvesTab", "Transform Curves"))
-		.SetGroup(WorkspaceMenuCategory.ToSharedRef())
-		.SetIcon(CurveIcon);
-}
-
-void FModulationPatchEditor::UnregisterTabSpawners(const TSharedRef<FTabManager>& InTabManager)
+TUniquePtr<WaveTable::Editor::FWaveTableCurveModel> FModulationPatchEditor::ConstructCurveModel(FRichCurve& InRichCurve, UObject* InParentObject, EWaveTableCurveSource InSource)
 {
-	InTabManager->UnregisterTabSpawner(PropertiesTabId);
-	InTabManager->UnregisterTabSpawner(CurveTabId);
+	using namespace WaveTable::Editor;
+	return TUniquePtr<FWaveTableCurveModel>(new FModPatchCurveEditorModel(InRichCurve, GetEditingObject(), InSource));
 }
 
-void FModulationPatchEditor::Init(const EToolkitMode::Type Mode, const TSharedPtr<IToolkitHost>& InitToolkitHost, USoundModulationPatch* InPatchToEdit)
+EWaveTableResolution FModulationPatchEditor::GetBankResolution() const
 {
-	CurveEditor = MakeShared<FCurveEditor>();
-	FCurveEditorInitParams InitParams;
-	CurveEditor->InitCurveEditor(InitParams);
-	CurveEditor->GridLineLabelFormatXAttribute = LOCTEXT("GridXLabelFormat", "{0}");
-
-	TUniquePtr<ICurveEditorBounds> EditorBounds = MakeUnique<FStaticCurveEditorBounds>();
-	EditorBounds->SetInputBounds(0.05, 1.05);
-	CurveEditor->SetBounds(MoveTemp(EditorBounds));
-
-	CurvePanel = SNew(SCurveEditorPanel, CurveEditor.ToSharedRef());
-
-	// Support undo/redo
-	InPatchToEdit->SetFlags(RF_Transactional);
-	GEditor->RegisterForUndo(this);
-
-	FDetailsViewArgs Args;
-	Args.bHideSelectionTip = true;
-	Args.NotifyHook = this;
-
-	FPropertyEditorModule& PropertyModule = FModuleManager::LoadModuleChecked<FPropertyEditorModule>("PropertyEditor");
-	PropertiesView = PropertyModule.CreateDetailView(Args);
-	PropertiesView->SetObject(InPatchToEdit);
-
-	TSharedRef<FTabManager::FLayout> StandaloneDefaultLayout = FTabManager::NewLayout("Standalone_ModulationPatchEditor_Layout_v2")
-		->AddArea
-		(
-			FTabManager::NewPrimaryArea()
-			->SetOrientation(Orient_Vertical)
-			->Split
-			(
-				FTabManager::NewSplitter()
-				->SetSizeCoefficient(0.9f)
-				->SetOrientation(Orient_Horizontal)
-				->Split
-				(
-					FTabManager::NewStack()
-					->SetSizeCoefficient(0.225f)
-					->AddTab(PropertiesTabId, ETabState::OpenedTab)
-				)
-				->Split
-				(
-					FTabManager::NewSplitter()
-					->SetSizeCoefficient(0.775f)
-					->SetOrientation(Orient_Vertical)
-					->Split
-					(
-						FTabManager::NewStack()
-						->SetHideTabWell(true)
-						->SetSizeCoefficient(0.33f)
-						->AddTab(CurveTabId, ETabState::OpenedTab)
-					)
-				)
-			)
-		);
-
-	const bool bCreateDefaultStandaloneMenu = true;
-	const bool bCreateDefaultToolbar = true;
-	const bool bToolbarFocusable = false;
-	const bool bUseSmallIcons = true;
-	FAssetEditorToolkit::InitAssetEditor(
-		Mode,
-		InitToolkitHost,
-		AppIdentifier,
-		StandaloneDefaultLayout,
-		bCreateDefaultStandaloneMenu,
-		bCreateDefaultToolbar,
-		InPatchToEdit,
-		bToolbarFocusable,
-		bUseSmallIcons);
-
-	FAudioModulationEditorModule& AudioEditorModule = FModuleManager::LoadModuleChecked<FAudioModulationEditorModule>(TEXT("AudioModulationEditor"));
-	AddMenuExtender(AudioEditorModule.GetModulationPatchMenuExtensibilityManager()->GetAllExtenders(GetToolkitCommands(), GetEditingObjects()));
-	AddToolbarExtender(AudioEditorModule.GetModulationPatchToolbarExtensibilityManager()->GetAllExtenders(GetToolkitCommands(), GetEditingObjects()));
-	AddToolbarExtender(CurvePanel->GetToolbarExtender());
-
-	if (CurveEditor.IsValid())
-	{
-		RegenerateMenusAndToolbars();
-	}
-}
-
-void FModulationPatchEditor::ClearExpressionCurve(int32 InInputIndex)
-{
-	if (CurveData.IsValidIndex(InInputIndex))
-	{
-		CurveData[InInputIndex].ExpressionCurve.Reset();
-	}
-}
-
-void FModulationPatchEditor::GenerateExpressionCurve(int32 InInputIndex, EModPatchOutputEditorCurveSource InSource, bool bInIsUnset)
-{
-	if (!CurveEditor.IsValid())
-	{
-		return;
-	}
-
-	USoundModulationPatch* Patch = CastChecked<USoundModulationPatch>(GetEditingObject());
-	if (!ensure(Patch) || !CurveData.IsValidIndex(InInputIndex))
-	{
-		return;
-	}
-
-	TSharedPtr<FRichCurve> Curve = CurveData[InInputIndex].ExpressionCurve;
-	if (!Curve.IsValid())
-	{
-		Curve = MakeShared<FRichCurve>();
-		CurveData[InInputIndex].ExpressionCurve = Curve;
-	}
-
-	if (!GetIsBypassed())
-	{
-		const FSoundControlModulationInput& Input = Patch->PatchSettings.Inputs[InInputIndex];
-
-		int32 CurveResolution;
-		switch (Input.Transform.Curve)
-		{
-			case ESoundModulatorCurve::Linear:
-			{
-				CurveResolution = 2;
-			}
-			break;
-
-			case ESoundModulatorCurve::Sin:
-			case ESoundModulatorCurve::SCurve:
-			{
-				CurveResolution = 64;
-			}
-			break;
-
-			case ESoundModulatorCurve::Log:
-			case ESoundModulatorCurve::Exp:
-			case ESoundModulatorCurve::Exp_Inverse:
-			{
-				CurveResolution = 256;
-			}
-			break;
-
-			default:
-			{
-				CurveResolution = 128;
-			}
-			break;
-		}
-
-		Curve->Reset();
-
-		check(CurveResolution > 1);
-		const float CurveResolutionRatio = 1.0f / (CurveResolution - 1);
-		for (int32 i = 0; i < CurveResolution; ++i)
-		{
-			const float X = FMath::Lerp(0.0f, 1.0f, CurveResolutionRatio * i);
-			float Y = X;
-			Input.Transform.Apply(Y);
-			Curve->AddKey(X, Y);
-		}
-	}
-
-	const EModPatchOutputEditorCurveSource Source = bInIsUnset ? EModPatchOutputEditorCurveSource::Unset : EModPatchOutputEditorCurveSource::Expression;
-	SetCurve(InInputIndex, *Curve.Get(), Source);
-}
-
-bool FModulationPatchEditor::RequiresNewCurve(int32 InInputIndex, const FRichCurve& InRichCurve) const
-{
-	const FCurveModelID CurveModelID = CurveData[InInputIndex].ModelID;
-	const TUniquePtr<FCurveModel>* CurveModel = CurveEditor->GetCurves().Find(CurveModelID);
-	if (!CurveModel || !CurveModel->IsValid())
-	{
-		return true;
-	}
-
-	FModPatchCurveEditorModel* PatchCurveModel = static_cast<FModPatchCurveEditorModel*>(CurveModel->Get());
-	check(PatchCurveModel);
-	if (&PatchCurveModel->GetRichCurve() != &InRichCurve)
+	if (USoundModulationPatch* Patch = Cast<USoundModulationPatch>(GetEditingObject()))
 	{
-		return true;
+		return Patch->PatchSettings.WaveTableResolution;
 	}
 
-	return false;
+	return EWaveTableResolution::None;
 }
 
-
-void FModulationPatchEditor::SetCurve(int32 InInputIndex, FRichCurve& InRichCurve, EModPatchOutputEditorCurveSource InSource)
+bool FModulationPatchEditor::GetIsPropertyEditorDisabled() const
 {
-	check(CurveEditor.IsValid());
-
-	if (!ensure(CurveData.IsValidIndex(InInputIndex)))
-	{
-		return;
-	}
-
-	FCurveData& CurveDataEntry = CurveData[InInputIndex];
-
-	const bool bRequiresNewCurve = RequiresNewCurve(InInputIndex, InRichCurve);
-	if (bRequiresNewCurve)
-	{
-		TUniquePtr<FModPatchCurveEditorModel> NewCurve = MakeUnique<FModPatchCurveEditorModel>(InRichCurve, GetEditingObject(), InSource, InInputIndex);
-		CurveDataEntry.ModelID = CurveEditor->AddCurve(MoveTemp(NewCurve));
-	}
-	else
-	{
-		const TUniquePtr<FCurveModel>& CurveModel = CurveEditor->GetCurves().FindChecked(CurveDataEntry.ModelID);
-		check(CurveModel.Get());
-		static_cast<FModPatchCurveEditorModel*>(CurveModel.Get())->Refresh(InSource, InInputIndex);
-	}
-
-	CurveEditor->PinCurve(CurveDataEntry.ModelID);
+	return GetIsBypassed();
 }
 
 bool FModulationPatchEditor::GetIsBypassed() const
@@ -287,264 +45,27 @@ bool FModulationPatchEditor::GetIsBypassed() const
 	return false;
 }
 
-FName FModulationPatchEditor::GetToolkitFName() const
-{
-	return FName("ModulationPatchEditor");
-}
-
-FText FModulationPatchEditor::GetBaseToolkitName() const
-{
-	return LOCTEXT( "AppLabel", "Modulation Patch Editor" );
-}
-
-FString FModulationPatchEditor::GetWorldCentricTabPrefix() const
+int32 FModulationPatchEditor::GetNumTransforms() const
 {
-	return LOCTEXT("WorldCentricTabPrefix", "ModulationPatch ").ToString();
-}
-
-FLinearColor FModulationPatchEditor::GetWorldCentricTabColorScale() const
-{
-	return FLinearColor( 0.0f, 0.0f, 0.2f, 0.5f );
-}
-
-EOrientation FModulationPatchEditor::GetSnapLabelOrientation() const
-{
-	return FMultiBoxSettings::UseSmallToolBarIcons.Get()
-		? EOrientation::Orient_Horizontal
-		: EOrientation::Orient_Vertical;
-}
-
-void FModulationPatchEditor::NotifyPostChange(const FPropertyChangedEvent& PropertyChangedEvent, FProperty* PropertyThatChanged)
-{
-	if (PropertyChangedEvent.ChangeType != EPropertyChangeType::Interactive)
-	{
-		RefreshCurves();
-	}
-}
-
-TSharedRef<SDockTab> FModulationPatchEditor::SpawnTab_Properties(const FSpawnTabArgs& Args)
-{
-	check(Args.GetTabId() == PropertiesTabId);
-
-	return SNew(SDockTab)
-		.Label(LOCTEXT("SoundModulationPatchDetailsTitle", "Details"))
-		[
-			PropertiesView.ToSharedRef()
-		];
-}
-
-TSharedRef<SDockTab> FModulationPatchEditor::SpawnTab_OutputCurve(const FSpawnTabArgs& Args)
-{
-	RefreshCurves();
-	CurveEditor->ZoomToFit();
-
-	TSharedRef<SDockTab> NewDockTab = SNew(SDockTab)
-		.Label(FText::Format(LOCTEXT("ModulationPatchFilterTitle", "Filter Transform Curve: {0}"), FText::FromString(GetEditingObject()->GetName())))
-		.TabColorScale(GetTabColorScale())
-		[
-			SNew(SBorder)
-			.BorderImage(FAppStyle::GetBrush("ToolPanel.GroupBorder"))
-			.Padding(0.0f)
-			[
-				CurvePanel.ToSharedRef()
-			]
-		];
-
-		return NewDockTab;
-}
-
-void FModulationPatchEditor::PostUndo(bool bSuccess)
-{
-	if (bSuccess)
-	{
-		RefreshCurves();
-	}
-}
-
-void FModulationPatchEditor::ResetCurves()
-{
-	const USoundModulationPatch* Patch = CastChecked<USoundModulationPatch>(GetEditingObject());
-	check(Patch);
-
-	CurveEditor->RemoveAllCurves();
-	CurveData.Reset();
-	CurveData.AddDefaulted(Patch->PatchSettings.Inputs.Num());
-}
-
-void FModulationPatchEditor::InitCurves()
-{
-	const USoundModulationPatch * Patch = CastChecked<USoundModulationPatch>(GetEditingObject());
-	const FSoundControlModulationPatch& PatchSettings = Patch->PatchSettings;
-	const int32 NumInputs = PatchSettings.Inputs.Num();
-	if (NumInputs < CurveData.Num() || NumInputs > CurveData.Num())
+	if (USoundModulationPatch* Patch = Cast<USoundModulationPatch>(GetEditingObject()))
 	{
-		ResetCurves();
-		return;
+		return Patch->PatchSettings.Inputs.Num();
 	}
 
-	for (int32 i = 0; i < PatchSettings.Inputs.Num(); ++i)
-	{
-		const FSoundControlModulationInput& Input = PatchSettings.Inputs[i];
-		switch (Input.Transform.Curve)
-		{
-			case ESoundModulatorCurve::Exp:
-			case ESoundModulatorCurve::Exp_Inverse:
-			case ESoundModulatorCurve::Linear:
-			case ESoundModulatorCurve::Log:
-			case ESoundModulatorCurve::SCurve:
-			case ESoundModulatorCurve::Sin:
-			{
-				if (RequiresNewCurve(i, *CurveData[i].ExpressionCurve.Get()))
-				{
-					ResetCurves();
-				}
-			}
-			break;
-
-			case ESoundModulatorCurve::Shared:
-			{
-				if (UCurveFloat* SharedCurve = Input.Transform.CurveShared)
-				{
-					if (RequiresNewCurve(i, SharedCurve->FloatCurve))
-					{
-						ResetCurves();
-					}
-				}
-				else if (RequiresNewCurve(i, *CurveData[i].ExpressionCurve.Get()))
-				{
-					ResetCurves();
-				}
-			}
-			break;
-
-			case ESoundModulatorCurve::Custom:
-			{
-				if (RequiresNewCurve(i, Input.Transform.CurveCustom))
-				{
-					ResetCurves();
-				}
-			}
-			break;
-
-			default:
-			{
-				static_assert(static_cast<int32>(ESoundModulatorCurve::Count) == 8, "Possible missing case coverage for output curve.");
-			}
-			break;
-		}
-	}
+	return 0;
 }
 
-void FModulationPatchEditor::RefreshCurves()
+FWaveTableTransform* FModulationPatchEditor::GetTransform(int32 InTransformIndex) const
 {
-	check(CurveEditor.IsValid());
-
-	for (const FCurveData& CurveDataEntry : CurveData)
-	{
-		CurveEditor->UnpinCurve(CurveDataEntry.ModelID);
-	}
-
-	USoundModulationPatch* Patch = CastChecked<USoundModulationPatch>(GetEditingObject());
-	check(Patch);
-
-	if (Patch->PatchSettings.bBypass)
-	{
-		ResetCurves();
-		return;
-	}
-
-	InitCurves();
-
-	for (int32 i = 0; i < Patch->PatchSettings.Inputs.Num(); ++i)
-	{
-		FSoundControlModulationInput& Input = Patch->PatchSettings.Inputs[i];
-		switch (Input.Transform.Curve)
-		{
-			case ESoundModulatorCurve::Exp:
-			case ESoundModulatorCurve::Exp_Inverse:
-			case ESoundModulatorCurve::Linear:
-			case ESoundModulatorCurve::Log:
-			case ESoundModulatorCurve::SCurve:
-			case ESoundModulatorCurve::Sin:
-			{
-				GenerateExpressionCurve(i, EModPatchOutputEditorCurveSource::Expression);
-			}
-			break;
-
-			case ESoundModulatorCurve::Shared:
-			{
-				if (UCurveFloat* SharedCurve = Input.Transform.CurveShared)
-				{
-					ClearExpressionCurve(i);
-					SetCurve(i, SharedCurve->FloatCurve, EModPatchOutputEditorCurveSource::Shared);
-				}
-				else
-				{
-					// Builds a dummy expression that just maps input to output in case
-					// where asset isn't selected and leave source as unset
-					GenerateExpressionCurve(i, EModPatchOutputEditorCurveSource::Expression, true /* bIsUnset */);
-				}
-			}
-			break;
-
-			case ESoundModulatorCurve::Custom:
-			{
-				TrimKeys(Input.Transform);
-				ClearExpressionCurve(i);
-				SetCurve(i, Input.Transform.CurveCustom, EModPatchOutputEditorCurveSource::Custom);
-			}
-			break;
-
-			default:
-			{
-				static_assert(static_cast<int32>(ESoundModulatorCurve::Count) == 8, "Possible missing case coverage for output curve.");
-			}
-			break;
-		}
-	}
-
-	// Collect and remove stale curves from editor
-	TArray<FCurveModelID> ToRemove;
-	TSet<FCurveModelID> ActiveModelIDs;
-	for (const FCurveData& CurveDataEntry : CurveData)
-	{
-		ActiveModelIDs.Add(CurveDataEntry.ModelID);
-	}
-
-	const TMap<FCurveModelID, TUniquePtr<FCurveModel>>& Curves = CurveEditor->GetCurves();
-	for (const TPair<FCurveModelID, TUniquePtr<FCurveModel>>& Pair : Curves)
+	if (USoundModulationPatch* Patch = Cast<USoundModulationPatch>(GetEditingObject()))
 	{
-		if (!ActiveModelIDs.Contains(Pair.Key))
+		FSoundControlModulationPatch& Settings = Patch->PatchSettings;
+		if (Settings.Inputs.IsValidIndex(InTransformIndex))
 		{
-			ToRemove.Add(Pair.Key);
+			return &Settings.Inputs[InTransformIndex].Transform;
 		}
 	}
 
-	for (FCurveModelID ModelID : ToRemove)
-	{
-		CurveEditor->RemoveCurve(ModelID);
-	}
-}
-
-void FModulationPatchEditor::PostRedo(bool bSuccess)
-{
-	if (bSuccess)
-	{
-		RefreshCurves();
-	}
-}
-
-void FModulationPatchEditor::TrimKeys(FSoundModulationTransform& OutTransform)
-{
-	FRichCurve& Curve = OutTransform.CurveCustom;
-	while (Curve.GetNumKeys() > 0 && 0.0f > Curve.GetFirstKey().Time)
-	{
-		Curve.DeleteKey(Curve.GetFirstKeyHandle());
-	}
-
-	while (Curve.GetNumKeys() > 0 && 1.0f < Curve.GetLastKey().Time)
-	{
-		Curve.DeleteKey(Curve.GetLastKeyHandle());
-	}
+	return nullptr;
 }
 #undef LOCTEXT_NAMESPACE
diff --git a/Engine/Plugins/Runtime/AudioModulation/Source/AudioModulationEditor/Private/Editors/ModulationPatchEditor.h b/Engine/Plugins/Runtime/AudioModulation/Source/AudioModulationEditor/Private/Editors/ModulationPatchEditor.h
--- a/Engine/Plugins/Runtime/AudioModulation/Source/AudioModulationEditor/Private/Editors/ModulationPatchEditor.h
+++ b/Engine/Plugins/Runtime/AudioModulation/Source/AudioModulationEditor/Private/Editors/ModulationPatchEditor.h
@@ -2,108 +2,23 @@
 
 #pragma once
 
-#include "CoreMinimal.h"
 
-#include "CurveEditorTypes.h"
-#include "Curves/CurveOwnerInterface.h"
-#include "EditorUndoClient.h"
-#include "Framework/Docking/TabManager.h"
-#include "Math/Color.h"
-#include "Misc/NotifyHook.h"
-#include "ModulationPatchCurveEditorViewStacked.h"
-#include "SoundModulationPatch.h"
-#include "Toolkits/AssetEditorToolkit.h"
-#include "Toolkits/IToolkitHost.h"
-#include "Widgets/SWidget.h"
+#include "WaveTableBankEditor.h"
 
 
-// Forward Declarations
-class FCurveEditor;
-class IToolkitHost;
-class SCurveEditorPanel;
-class UCurveBase;
-class USoundModulationPatch;
-
-struct FSoundControlModulationInput;
-
-
-class FModulationPatchEditor : public FAssetEditorToolkit, public FNotifyHook, public FEditorUndoClient
+class FModulationPatchEditor : public WaveTable::Editor::FBankEditorBase
 {
 public:
-	FModulationPatchEditor();
+	FModulationPatchEditor() = default;
 	virtual ~FModulationPatchEditor() = default;
 
-	void Init(const EToolkitMode::Type Mode, const TSharedPtr<IToolkitHost>& InitToolkitHost, USoundModulationPatch* PatchToEdit);
-
-	/** FAssetEditorToolkit interface */
-	virtual FName GetToolkitFName() const override;
-	virtual FText GetBaseToolkitName() const override;
-	virtual FString GetWorldCentricTabPrefix() const override;
-	virtual FLinearColor GetWorldCentricTabColorScale() const override;
-	virtual void RegisterTabSpawners(const TSharedRef<FTabManager>& TabManager) override;
-	virtual void UnregisterTabSpawners(const TSharedRef<FTabManager>& TabManager) override;
-
-	/** FNotifyHook interface */
-	virtual void NotifyPostChange(const FPropertyChangedEvent& PropertyChangedEvent, FProperty* PropertyThatChanged) override;
-
 protected:
-	virtual void PostUndo(bool bSuccess) override;
-	virtual void PostRedo(bool bSuccess) override;
-
-private:
-	void SetCurve(int32 InInputIndex, FRichCurve& InRichCurve, EModPatchOutputEditorCurveSource InSource);
-
-	/**	Spawns the tab allowing for editing/viewing the output curve(s) */
-	TSharedRef<SDockTab> SpawnTab_OutputCurve(const FSpawnTabArgs& Args);
-
-	/**	Spawns the tab allowing for editing/viewing the output curve(s) */
-	TSharedRef<SDockTab> SpawnTab_Properties(const FSpawnTabArgs& Args);
-
-	bool GetIsBypassed() const;
-
-	/** Get the orientation for the snap value controls. */
-	EOrientation GetSnapLabelOrientation() const;
-
-	/** Updates patch's input curves. */
-	void RefreshCurves();
-
-	/** Trims keys out-of-bounds in provided output transform's curve */
-	static void TrimKeys(FSoundModulationTransform& OutTransform);
-
-	/** Clears the expression curve at the given input index */
-	void ClearExpressionCurve(int32 InInputIndex);
-
-	/** Generates expression curve at the given index. */
-	void GenerateExpressionCurve(int32 InInputIndex, EModPatchOutputEditorCurveSource InSource, bool bInIsUnset = false);
-
-	void InitCurves();
-
-	void ResetCurves();
-
-	bool RequiresNewCurve(int32 InInputIndex, const FRichCurve& InRichCurve) const;
-
-	TSharedPtr<FUICommandList> ToolbarCurveTargetCommands;
-
-	TSharedPtr<FCurveEditor> CurveEditor;
-	TSharedPtr<SCurveEditorPanel> CurvePanel;
-
-	struct FCurveData
-	{
-		FCurveModelID ModelID;
-		TSharedPtr<FRichCurve> ExpressionCurve;
-
-		FCurveData()
-			: ModelID(FCurveModelID::Unique())
-		{
-		}
-	};
-	TArray<FCurveData> CurveData;
-
-	/** Properties tab */
-	TSharedPtr<IDetailsView> PropertiesView;
-
-	/** Settings Editor App Identifier */
-	static const FName AppIdentifier;
-	static const FName CurveTabId;
-	static const FName PropertiesTabId;
+	virtual bool GetIsBypassed() const;
+	virtual bool GetIsPropertyEditorDisabled() const override;
+	virtual EWaveTableResolution GetBankResolution() const override;
+	virtual bool GetBankIsBipolar() const override { return false; }
+	virtual FWaveTableTransform* GetTransform(int32 InIndex) const override;
+	virtual int32 GetNumTransforms() const override;
+
+	virtual TUniquePtr<WaveTable::Editor::FWaveTableCurveModel> ConstructCurveModel(FRichCurve& InRichCurve, UObject* InParentObject, EWaveTableCurveSource InSource) override;
 };
diff --git a/Engine/Plugins/Runtime/AudioModulation/Source/AudioModulationEditor/Private/Layouts/SoundControlModulationPatchLayout.cpp b/Engine/Plugins/Runtime/AudioModulation/Source/AudioModulationEditor/Private/Layouts/SoundControlModulationPatchLayout.cpp
--- a/Engine/Plugins/Runtime/AudioModulation/Source/AudioModulationEditor/Private/Layouts/SoundControlModulationPatchLayout.cpp
+++ b/Engine/Plugins/Runtime/AudioModulation/Source/AudioModulationEditor/Private/Layouts/SoundControlModulationPatchLayout.cpp
@@ -72,6 +72,7 @@ void FSoundControlModulationPatchLayoutCustomization::CustomizeChildren(TSharedR
 	AudioModulationEditorUtils::GetPropertyHandleMap(StructPropertyHandle, PropertyHandles);
 
 	TSharedRef<IPropertyHandle> BypassHandle = PropertyHandles.FindChecked(GET_MEMBER_NAME_CHECKED(FSoundControlModulationPatch, bBypass)).ToSharedRef();
+	TSharedRef<IPropertyHandle> WaveTableResolutionHandle = PropertyHandles.FindChecked(GET_MEMBER_NAME_CHECKED(FSoundControlModulationPatch, WaveTableResolution)).ToSharedRef();
 	TSharedRef<IPropertyHandle> InputsHandle = PropertyHandles.FindChecked(GET_MEMBER_NAME_CHECKED(FSoundControlModulationPatch, Inputs)).ToSharedRef();
 	TSharedRef<IPropertyHandle> ParameterHandle = PropertyHandles.FindChecked(GET_MEMBER_NAME_CHECKED(FSoundControlModulationPatch, OutputParameter)).ToSharedRef();
 
@@ -83,6 +84,7 @@ void FSoundControlModulationPatchLayoutCustomization::CustomizeChildren(TSharedR
 	});
 
 	ChildBuilder.AddProperty(BypassHandle);
+	ChildBuilder.AddProperty(WaveTableResolutionHandle);
 	ChildBuilder.AddProperty(ParameterHandle).EditCondition(EditCondition, nullptr);
 	ChildBuilder.AddProperty(InputsHandle).EditCondition(EditCondition, nullptr);
 }
diff --git a/Engine/Plugins/Runtime/AudioModulation/Source/AudioModulationEditor/Private/Layouts/SoundModulationTransformLayout.cpp b/Engine/Plugins/Runtime/AudioModulation/Source/AudioModulationEditor/Private/Layouts/SoundModulationTransformLayout.cpp
--- a/Engine/Plugins/Runtime/AudioModulation/Source/AudioModulationEditor/Private/Layouts/SoundModulationTransformLayout.cpp
+++ b/Engine/Plugins/Runtime/AudioModulation/Source/AudioModulationEditor/Private/Layouts/SoundModulationTransformLayout.cpp
@@ -1,96 +1,57 @@
 // Copyright Epic Games, Inc. All Rights Reserved.
 #include "SoundModulationTransformLayout.h"
 
-#include "Audio.h"
-#include "AudioModulationStyle.h"
-#include "DetailCategoryBuilder.h"
-#include "DetailLayoutBuilder.h"
-#include "DetailWidgetRow.h"
-#include "Styling/AppStyle.h"
-#include "IDetailChildrenBuilder.h"
-#include "IDetailPropertyRow.h"
-#include "InputCoreTypes.h"
-#include "PropertyRestriction.h"
-#include "SCurveEditor.h"
-#include "ScopedTransaction.h"
+#include "Containers/Array.h"
 #include "SoundModulationPatch.h"
-#include "SoundModulationTransform.h"
-#include "Widgets/Input/SButton.h"
-#include "Widgets/Input/SEditableTextBox.h"
-#include "Widgets/Input/SNumericEntryBox.h"
-#include "Widgets/Layout/SSeparator.h"
-#include "Widgets/Layout/SWidgetSwitcher.h"
-#include "Widgets/Notifications/SNotificationList.h"
-#include "Widgets/Text/SRichTextBlock.h"
-#include "Widgets/Text/STextBlock.h"
-#include "Widgets/Views/SListView.h"
+#include "UObject/NoExportTypes.h"
 
 
-#define LOCTEXT_NAMESPACE "SoundModulationOutputTransform"
-void FSoundModulationTransformLayoutCustomization::CustomizeHeader(TSharedRef<IPropertyHandle> StructPropertyHandle, FDetailWidgetRow& HeaderRow, IPropertyTypeCustomizationUtils& CustomizationUtils)
+const EWaveTableResolution* FSoundModulationTransformLayoutCustomization::GetResolution() const
 {
-}
-
-void FSoundModulationTransformLayoutCustomization::CustomizeChildren(TSharedRef<IPropertyHandle> StructPropertyHandle, IDetailChildrenBuilder& ChildBuilder, IPropertyTypeCustomizationUtils& StructCustomizationUtils)
-{
-	uint32 NumChildren;
-	StructPropertyHandle->GetNumChildren(NumChildren);
-
-	TMap<FName, TSharedPtr<IPropertyHandle>> PropertyHandles;
-
-	for (uint32 ChildIndex = 0; ChildIndex < NumChildren; ++ChildIndex)
+	if (ensure(WaveTableOptionsHandle.IsValid()))
 	{
-		TSharedRef<IPropertyHandle> ChildHandle = StructPropertyHandle->GetChildHandle(ChildIndex).ToSharedRef();
-		const FName PropertyName = ChildHandle->GetProperty()->GetFName();
-		PropertyHandles.Add(PropertyName, ChildHandle);
+		TArray<UObject*> OuterObjects;
+		WaveTableOptionsHandle->GetOuterObjects(OuterObjects);
+		if (OuterObjects.Num() == 1)
+		{
+			if (USoundModulationPatch* Patch = Cast<USoundModulationPatch>(OuterObjects.Last()))
+			{
+				FSoundControlModulationPatch& PatchSettings = Patch->PatchSettings;
+				return &PatchSettings.WaveTableResolution;
+			}
+		}
 	}
 
-	TSharedRef<IPropertyHandle>CurveHandle = PropertyHandles.FindChecked(GET_MEMBER_NAME_CHECKED(FSoundModulationTransform, Curve)).ToSharedRef();
-	ChildBuilder.AddProperty(CurveHandle);
-
-	TSharedRef<IPropertyHandle> ScalarHandle = PropertyHandles.FindChecked(GET_MEMBER_NAME_CHECKED(FSoundModulationTransform, Scalar)).ToSharedRef();
-
-	ChildBuilder.AddProperty(ScalarHandle)
-		.EditCondition(TAttribute<bool>::Create([this, CurveHandle]() { return IsScaleableCurve(CurveHandle); }), nullptr)
-		.Visibility(TAttribute<EVisibility>::Create([this, CurveHandle]() { return IsScaleableCurve(CurveHandle) ? EVisibility::Visible : EVisibility::Hidden; }));
-
-	TSharedRef<IPropertyHandle> SharedCurveHandle = PropertyHandles.FindChecked(GET_MEMBER_NAME_CHECKED(FSoundModulationTransform, CurveShared)).ToSharedRef();
-	ChildBuilder.AddProperty(SharedCurveHandle)
-		.EditCondition(TAttribute<bool>::Create([this, CurveHandle]() { return IsSharedCurve(CurveHandle); }), nullptr)
-		.Visibility(TAttribute<EVisibility>::Create([this, CurveHandle]() { return IsSharedCurve(CurveHandle) ? EVisibility::Visible : EVisibility::Hidden; }));
+	return nullptr;
 }
 
-bool FSoundModulationTransformLayoutCustomization::IsScaleableCurve(TSharedPtr<IPropertyHandle> CurveHandle) const
+FWaveTableTransform* FSoundModulationTransformLayoutCustomization::GetTransform() const
 {
-	if (!CurveHandle.IsValid())
+	if (ensure(WaveTableOptionsHandle.IsValid()))
 	{
-		return false;
-	}
-
-	uint8 CurveString;
-	CurveHandle->GetValue(CurveString);
-
-	static const TArray<ESoundModulatorCurve> ScalarFilters = { ESoundModulatorCurve::Exp, ESoundModulatorCurve::Exp_Inverse, ESoundModulatorCurve::Log };
-	for (ESoundModulatorCurve Filter : ScalarFilters)
-	{
-		if (CurveString == static_cast<uint8>(Filter))
+		TArray<UObject*> OuterObjects;
+		WaveTableOptionsHandle->GetOuterObjects(OuterObjects);
+		if (OuterObjects.Num() == 1)
 		{
-			return true;
+			if (USoundModulationPatch* Patch = Cast<USoundModulationPatch>(OuterObjects.Last()))
+			{
+				const int32 InputIndex = GetOwningArrayIndex();
+				if (InputIndex != INDEX_NONE)
+				{
+					FSoundControlModulationPatch& PatchSettings = Patch->PatchSettings;
+					if (InputIndex < PatchSettings.Inputs.Num())
+					{
+						return &PatchSettings.Inputs[InputIndex].Transform;
+					}
+				}
+			}
 		}
 	}
 
-	return false;
+	return nullptr;
 }
 
-bool FSoundModulationTransformLayoutCustomization::IsSharedCurve(TSharedPtr<IPropertyHandle> CurveHandle) const
+bool FSoundModulationTransformLayoutCustomization::IsBipolar() const
 {
-	if (!CurveHandle.IsValid())
-	{
-		return false;
-	}
-
-	uint8 CurveValue;
-	CurveHandle->GetValue(CurveValue);
-	return CurveValue == static_cast<uint8>(ESoundModulatorCurve::Shared);
+	return false;
 }
-#undef LOCTEXT_NAMESPACE
diff --git a/Engine/Plugins/Runtime/AudioModulation/Source/AudioModulationEditor/Private/Layouts/SoundModulationTransformLayout.h b/Engine/Plugins/Runtime/AudioModulation/Source/AudioModulationEditor/Private/Layouts/SoundModulationTransformLayout.h
--- a/Engine/Plugins/Runtime/AudioModulation/Source/AudioModulationEditor/Private/Layouts/SoundModulationTransformLayout.h
+++ b/Engine/Plugins/Runtime/AudioModulation/Source/AudioModulationEditor/Private/Layouts/SoundModulationTransformLayout.h
@@ -1,21 +1,21 @@
 // Copyright Epic Games, Inc. All Rights Reserved.
 #pragma once
 
-#include "CoreMinimal.h"
-
 #include "Curves/SimpleCurve.h"
 #include "IPropertyTypeCustomization.h"
 #include "Layout/Visibility.h"
 #include "Misc/Attribute.h"
 #include "PropertyHandle.h"
 #include "SCurveEditor.h"
-#include "SoundModulationTransform.h"
+#include "WaveTableTransformLayout.h"
 
 
 // Forward Declarations
+struct FSoundControlModulationInput;
 struct FSoundModulationTransform;
+class USoundModulationPatch;
 
-class FSoundModulationTransformLayoutCustomization : public IPropertyTypeCustomization
+class FSoundModulationTransformLayoutCustomization : public WaveTable::Editor::FTransformLayoutCustomizationBase
 {
 public:
 	static TSharedRef<IPropertyTypeCustomization> MakeInstance()
@@ -23,12 +23,10 @@ public:
 		return MakeShared<FSoundModulationTransformLayoutCustomization>();
 	}
 
-	//~ Begin IPropertyTypeCustomization
-	virtual void CustomizeHeader(TSharedRef<IPropertyHandle> StructPropertyHandle, FDetailWidgetRow& HeaderRow, IPropertyTypeCustomizationUtils& StructCustomizationUtils) override;
-	virtual void CustomizeChildren(TSharedRef<IPropertyHandle> StructPropertyHandle, IDetailChildrenBuilder& ChildBuilder, IPropertyTypeCustomizationUtils& StructCustomizationUtils) override;
-	//~ End IPropertyTypeCustomization
-
-private:
-	bool IsScaleableCurve(TSharedPtr<IPropertyHandle> CurveHandle) const;
-	bool IsSharedCurve(TSharedPtr<IPropertyHandle> CurveHandle) const;
+protected:
+	//~ Begin FTransformLayoutCustomizationBase
+	virtual const EWaveTableResolution* GetResolution() const override;
+	virtual FWaveTableTransform* GetTransform() const override;
+	virtual bool IsBipolar() const override;
+	//~ End FTransformLayoutCustomizationBase
 };
diff --git a/Engine/Plugins/Runtime/Metasound/Metasound.uplugin b/Engine/Plugins/Runtime/Metasound/Metasound.uplugin
--- a/Engine/Plugins/Runtime/Metasound/Metasound.uplugin
+++ b/Engine/Plugins/Runtime/Metasound/Metasound.uplugin
@@ -59,6 +59,10 @@
 		{
 			"Name": "AudioSynesthesia",
 			"Enabled": true
+		},
+		{
+			"Name": "WaveTable",
+			"Enabled": true
 		}
 	]
 }
diff --git a/Engine/Plugins/Runtime/Metasound/Source/MetasoundEngine/MetasoundEngine.Build.cs b/Engine/Plugins/Runtime/Metasound/Source/MetasoundEngine/MetasoundEngine.Build.cs
--- a/Engine/Plugins/Runtime/Metasound/Source/MetasoundEngine/MetasoundEngine.Build.cs
+++ b/Engine/Plugins/Runtime/Metasound/Source/MetasoundEngine/MetasoundEngine.Build.cs
@@ -9,6 +9,8 @@ namespace UnrealBuildTool.Rules
 	{
 		public MetasoundEngine(ReadOnlyTargetRules Target) : base(Target)
 		{
+			OptimizeCode = CodeOptimization.Never;
+
 			PrivateDependencyModuleNames.AddRange
 			(
 				new string[]
@@ -34,6 +36,7 @@ namespace UnrealBuildTool.Rules
 					"AudioCodecEngine",
 					"Serialization",
 					"AudioPlatformConfiguration",
+					"WaveTable"
 				}
 			);
 		}
diff --git a/Engine/Plugins/Runtime/Metasound/Source/MetasoundEngine/Private/MetasoundEngineModule.cpp b/Engine/Plugins/Runtime/Metasound/Source/MetasoundEngine/Private/MetasoundEngineModule.cpp
--- a/Engine/Plugins/Runtime/Metasound/Source/MetasoundEngine/Private/MetasoundEngineModule.cpp
+++ b/Engine/Plugins/Runtime/Metasound/Source/MetasoundEngine/Private/MetasoundEngineModule.cpp
@@ -10,12 +10,15 @@
 #include "MetasoundSource.h"
 #include "MetasoundUObjectRegistry.h"
 #include "MetasoundWave.h"
+#include "MetasoundWaveTable.h"
 #include "Modules/ModuleManager.h"
 #include "Sound/AudioSettings.h"
 
 DEFINE_LOG_CATEGORY(LogMetasoundEngine);
 
 REGISTER_METASOUND_DATATYPE(Metasound::FWaveAsset, "WaveAsset", Metasound::ELiteralType::UObjectProxy, USoundWave);
+REGISTER_METASOUND_DATATYPE(WaveTable::FWaveTable, "WaveTable", Metasound::ELiteralType::FloatArray)
+REGISTER_METASOUND_DATATYPE(Metasound::FWaveTableBankAsset, "WaveTableBankAsset", Metasound::ELiteralType::UObjectProxy, UWaveTableBank);
 
 class FMetasoundEngineModule : public IMetasoundEngineModule
 {
@@ -26,6 +29,7 @@ class FMetasoundEngineModule : public IMetasoundEngineModule
 		FModuleManager::Get().LoadModuleChecked("MetasoundStandardNodes");
 		FModuleManager::Get().LoadModuleChecked("MetasoundGenerator");
 		FModuleManager::Get().LoadModuleChecked("AudioCodecEngine");
+		FModuleManager::Get().LoadModuleChecked("WaveTable");
 
 		// Register engine-level parameter interfaces if not done already.
 		// (Potentially not already called if plugin is loaded while cooking.)
diff --git a/Engine/Plugins/Runtime/Metasound/Source/MetasoundEngine/Public/MetasoundWaveTable.h b/Engine/Plugins/Runtime/Metasound/Source/MetasoundEngine/Public/MetasoundWaveTable.h
new file mode 100644
--- /dev/null
+++ b/Engine/Plugins/Runtime/Metasound/Source/MetasoundEngine/Public/MetasoundWaveTable.h
@@ -0,0 +1,83 @@
+// Copyright Epic Games, Inc. All Rights Reserved.
+#pragma once
+
+#include "Containers/Array.h"
+#include "Containers/ArrayView.h"
+#include "IAudioProxyInitializer.h"
+#include "MetasoundDataReference.h"
+#include "MetasoundDataTypeRegistrationMacro.h"
+#include "MetasoundPrimitives.h"
+#include "WaveTable.h"
+#include "WaveTableBank.h"
+
+
+namespace Metasound
+{
+	class FWaveTableBankAsset;
+	using FWaveTableBankAssetReadRef = TDataReadReference<FWaveTableBankAsset>;
+	using FWaveTableReadRef = TDataReadReference<WaveTable::FWaveTable>;
+	using FWaveTableWriteRef = TDataWriteReference<WaveTable::FWaveTable>;
+
+	class METASOUNDENGINE_API FWaveTableBankAsset
+	{
+		FWaveTableBankAssetProxyPtr Proxy;
+
+	public:
+		FWaveTableBankAsset()
+			: Proxy(nullptr)
+		{
+		}
+
+		FWaveTableBankAsset(const FWaveTableBankAsset& InOther)
+			: Proxy(InOther.Proxy)
+		{
+		}
+
+		FWaveTableBankAsset& operator=(const FWaveTableBankAsset& InOther)
+		{
+			Proxy = InOther.Proxy;
+			return *this;
+		}
+
+		FWaveTableBankAsset& operator=(FWaveTableBankAsset&& InOther)
+		{
+			Proxy = MoveTemp(InOther.Proxy);
+			return *this;
+		}
+
+		FWaveTableBankAsset(const Audio::IProxyDataPtr& InInitData)
+		{
+			if (InInitData.IsValid())
+			{
+				if (InInitData->CheckTypeCast<FWaveTableBankAssetProxy>())
+				{
+					Proxy = MakeShared<FWaveTableBankAssetProxy, ESPMode::ThreadSafe>(InInitData->GetAs<FWaveTableBankAssetProxy>());
+				}
+			}
+		}
+
+		bool IsValid() const
+		{
+			return Proxy.IsValid();
+		}
+
+		const FWaveTableBankAssetProxyPtr& GetProxy() const
+		{
+			return Proxy;
+		}
+
+		const FWaveTableBankAssetProxy* operator->() const
+		{
+			return Proxy.Get();
+		}
+
+		FWaveTableBankAssetProxy* operator->()
+		{
+			return Proxy.Get();
+		}
+	};
+
+	DECLARE_METASOUND_DATA_REFERENCE_TYPES(WaveTable::FWaveTable, METASOUNDENGINE_API, FWaveTableTypeInfo, FWaveTableReadRef, FWaveTableWriteRef)
+	DECLARE_METASOUND_DATA_REFERENCE_TYPES(FWaveTableBankAsset, METASOUNDENGINE_API, FWaveTableBankAssetTypeInfo, FWaveTableBankAssetReadRef, FWaveTableBankAssetWriteRef)
+} // namespace Metasound
+
diff --git a/Engine/Plugins/Runtime/Metasound/Source/MetasoundStandardNodes/MetasoundStandardNodes.Build.cs b/Engine/Plugins/Runtime/Metasound/Source/MetasoundStandardNodes/MetasoundStandardNodes.Build.cs
--- a/Engine/Plugins/Runtime/Metasound/Source/MetasoundStandardNodes/MetasoundStandardNodes.Build.cs
+++ b/Engine/Plugins/Runtime/Metasound/Source/MetasoundStandardNodes/MetasoundStandardNodes.Build.cs
@@ -9,15 +9,17 @@ namespace UnrealBuildTool.Rules
 	{
 		public MetasoundStandardNodes(ReadOnlyTargetRules Target) : base(Target)
 		{
+			OptimizeCode = CodeOptimization.Never;
+
 			PublicDependencyModuleNames.AddRange
 			(
 				new string[]
 				{
+					"AudioExtensions",
 					"Core",
+					"MetasoundFrontend",
 					"Serialization",
-					"SignalProcessing",
-					"AudioExtensions",
-					"MetasoundFrontend"
+					"SignalProcessing"
 				}
 			);
 
@@ -25,8 +27,9 @@ namespace UnrealBuildTool.Rules
 			(
 				new string[]
 				{
+					"CoreUObject",
 					"MetasoundGraphCore",
-					"CoreUObject"
+					"WaveTable"
 				}
 			);
 		}
diff --git a/Engine/Plugins/Runtime/WaveTable/Source/WaveTable/Private/WaveTableBank.cpp b/Engine/Plugins/Runtime/WaveTable/Source/WaveTable/Private/WaveTableBank.cpp
new file mode 100644
--- /dev/null
+++ b/Engine/Plugins/Runtime/WaveTable/Source/WaveTable/Private/WaveTableBank.cpp
@@ -0,0 +1,54 @@
+// Copyright Epic Games, Inc. All Rights Reserved.
+#include "WaveTableBank.h"
+
+#define LOCTEXT_NAMESPACE "WaveTable"
+
+
+TUniquePtr<Audio::IProxyData> UWaveTableBank::CreateNewProxyData(const Audio::FProxyDataInitParams& InitParams)
+{
+	return MakeUnique<FWaveTableBankAssetProxy>(*this);
+}
+
+#if WITH_EDITOR
+void UWaveTableBank::PostEditChangeProperty(FPropertyChangedEvent& InPropertyChangedEvent)
+{
+	FProperty* Property = InPropertyChangedEvent.Property;
+	if (Property && InPropertyChangedEvent.ChangeType != EPropertyChangeType::Interactive)
+	{
+		for (FWaveTableBankEntry& Entry : Entries)
+		{
+			Entry.Transform.CacheWaveTable(Resolution, bBipolar);
+		}
+	}
+
+	Super::PostEditChangeProperty(InPropertyChangedEvent);
+}
+
+void UWaveTableBank::PostEditChangeChainProperty(FPropertyChangedChainEvent& InPropertyChangedEvent)
+{
+	FProperty* Property = InPropertyChangedEvent.Property;
+	if (Property && InPropertyChangedEvent.ChangeType != EPropertyChangeType::Interactive)
+	{
+		for (FWaveTableBankEntry& Entry : Entries)
+		{
+			Entry.Transform.CacheWaveTable(Resolution, bBipolar);
+		}
+	}
+
+	Super::PostEditChangeChainProperty(InPropertyChangedEvent);
+}
+
+void UWaveTableBank::PreSave(FObjectPreSaveContext InSaveContext)
+{
+	if (!InSaveContext.IsCooking())
+	{
+		for (FWaveTableBankEntry& Entry : Entries)
+		{
+			Entry.Transform.CacheWaveTable(Resolution, bBipolar);
+		}
+	}
+
+	Super::PreSave(InSaveContext);
+}
+#endif // WITH_EDITOR
+#undef LOCTEXT_NAMESPACE // WaveTable
diff --git a/Engine/Plugins/Runtime/WaveTable/Source/WaveTable/Private/WaveTableImporter.cpp b/Engine/Plugins/Runtime/WaveTable/Source/WaveTable/Private/WaveTableImporter.cpp
--- a/Engine/Plugins/Runtime/WaveTable/Source/WaveTable/Private/WaveTableImporter.cpp
+++ b/Engine/Plugins/Runtime/WaveTable/Source/WaveTable/Private/WaveTableImporter.cpp
@@ -9,6 +9,7 @@
 #include "UObject/ObjectMacros.h"
 #include "UObject/Package.h"
 #include "WaveTableSampler.h"
+#include "DSP/EnvelopeFollower.h"
 
 
 namespace WaveTable
@@ -24,27 +25,30 @@ namespace WaveTable
 		};
 	}
 
-	FImporter::FImporter(const FWaveTableSettings& InOptions, EWaveTableResolution InResolution)
-		: Options(InOptions)
+	FImporter::FImporter(const FWaveTableSettings& InSettings, EWaveTableResolution InResolution, bool bInBipolar)
+		: Settings(InSettings)
 		, Resolution(InResolution)
+		, bBipolar(bInBipolar)
 	{
 		Sampler.SetInterpolationMode(FWaveTableSampler::EInterpolationMode::Cubic);
-		Sampler.SetPhase(Options.Phase);
+		Sampler.SetPhase(Settings.Phase);
 	}
 
 	void FImporter::Process(TArray<float>& OutWaveTable)
 	{
 		using namespace ImporterPrivate;
 
+		const int32 ResolutionInt = ResolutionToInt32(Resolution);
+
 		OutWaveTable.Reset();
-		OutWaveTable.AddZeroed(ResolutionToInt32(Resolution));
+		OutWaveTable.AddZeroed(ResolutionInt);
 
 		if (Resolution == EWaveTableResolution::None)
 		{
 			return;
 		}
 
-		int32 SourceNumSamples = Options.SourcePCMData.Num();
+		int32 SourceNumSamples = Settings.SourcePCMData.Num();
 		if (SourceNumSamples == 0)
 		{
 			return;
@@ -52,61 +56,70 @@ namespace WaveTable
 
 		// 1. Use Top & Tail to determine start offset & sample count of SourceArrayView
 		int32 SourceTopOffset = 0;
-		if (Options.Top > 0.0f)
+		if (Settings.Top > 0.0f)
 		{
-			SourceTopOffset = ImporterPrivate::RatioToIndex(Options.Top, SourceNumSamples);
+			SourceTopOffset = ImporterPrivate::RatioToIndex(Settings.Top, SourceNumSamples);
 			SourceNumSamples -= SourceTopOffset;
 		}
 
-		if (Options.Tail > 0.0f)
+		if (Settings.Tail > 0.0f)
 		{
-			const int32 SourceTailOffset = RatioToIndex(Options.Tail, SourceNumSamples);
+			const int32 SourceTailOffset = RatioToIndex(Settings.Tail, SourceNumSamples);
 			SourceNumSamples -= SourceTailOffset;
 		}
 
-		if (SourceNumSamples > 0 && SourceTopOffset < Options.SourcePCMData.Num() && !OutWaveTable.IsEmpty())
+		if (SourceNumSamples > 0 && SourceTopOffset < Settings.SourcePCMData.Num() && !OutWaveTable.IsEmpty())
 		{
-			TArrayView<const float> SourceArrayView(Options.SourcePCMData.GetData() + SourceTopOffset, SourceNumSamples);
-			TArrayView<float> WaveTableView(OutWaveTable);
-			Sampler.Process(SourceArrayView, WaveTableView);
+			TArrayView<const float> SourceArrayView(Settings.SourcePCMData.GetData() + SourceTopOffset, SourceNumSamples);
 
-			if (Options.FadeIn > 0.0f)
-			{
-				const int32 FadeLength = RatioToIndex(Options.FadeIn, OutWaveTable.Num());
-				TArrayView<float> FadeView(OutWaveTable.GetData(), FadeLength);
-				Audio::ArrayFade(FadeView, 0.0f, 1.0f);
-			}
+			// 2. Resample into table
+			Sampler.Reset();
+			Sampler.Process(SourceArrayView, OutWaveTable);
 
-			if (Options.FadeOut > 0.0f)
+			// 3. Apply offset
+			float TableOffset = 0.0f;
+			if (Settings.bRemoveOffset || !bBipolar)
 			{
-				const int32 FadeLastIndex = RatioToIndex(Options.FadeOut, OutWaveTable.Num());
-				const int32 FadeInitIndex = RatioToIndex(1.0f - Options.FadeOut, OutWaveTable.Num());
-				TArrayView<float> FadeView = TArrayView<float>(&OutWaveTable[FadeInitIndex], FadeLastIndex + 1);
-				Audio::ArrayFade(FadeView, 1.0f, 0.0f);
+				TableOffset = Audio::ArrayGetAverageValue(OutWaveTable);
+				Audio::ArrayAddConstantInplace(OutWaveTable, -1.0f * TableOffset);
 			}
 
-			if (Options.bRemoveOffset)
+			// 4. Normalize
+			if (Settings.bNormalize)
 			{
-				const float TableOffset = Audio::ArrayGetAverageValue(OutWaveTable);
-				Audio::ArrayAddConstantInplace(OutWaveTable, -1.0f * TableOffset);
+				const float MaxValue = Audio::ArrayMaxAbsValue(OutWaveTable);
+				if (MaxValue > 0.0f)
+				{
+					Audio::ArrayMultiplyByConstantInPlace(OutWaveTable, 1.0f / MaxValue);
+				}
 			}
 
-			if (Options.bConvertToEnvelope)
+			// 6. Apply fades
+			if (Settings.FadeIn > 0.0f)
 			{
-				Audio::ArrayAbsInPlace(OutWaveTable);
+				const int32 FadeLength = RatioToIndex(Settings.FadeIn, OutWaveTable.Num());
+				TArrayView<float> FadeView(OutWaveTable.GetData(), FadeLength);
+				Audio::ArrayFade(FadeView, 0.0f, 1.0f);
 			}
-			else
+
+			if (Settings.FadeOut > 0.0f)
 			{
-				Audio::ArrayAddConstantInplace(OutWaveTable, 1.0f);
-				Audio::ArrayMultiplyByConstantInPlace(OutWaveTable, 0.5f);
+				const int32 FadeLastIndex = RatioToIndex(Settings.FadeOut, OutWaveTable.Num());
+				const int32 FadeInitIndex = RatioToIndex(1.0f - Settings.FadeOut, OutWaveTable.Num());
+				TArrayView<float> FadeView = TArrayView<float>(&OutWaveTable[FadeInitIndex], FadeLastIndex + 1);
+				Audio::ArrayFade(FadeView, 1.0f, 0.0f);
 			}
 
-			if (Options.bNormalize)
+			// 7. Finalize if unipolar source
+			if (!bBipolar)
 			{
-				const float MaxValue = Audio::ArrayMaxAbsValue(OutWaveTable);
-				if (MaxValue > 0.0f)
+				Audio::ArrayAbsInPlace(OutWaveTable);
+
+				// If not requesting offset removal, add back (is always removed above for
+				// unipolar to ensure fades are centered around offset)
+				if (!Settings.bRemoveOffset)
 				{
-					Audio::ArrayMultiplyByConstantInPlace(OutWaveTable, 1.0f / MaxValue);
+					Audio::ArrayAddConstantInplace(OutWaveTable, TableOffset);
 				}
 			}
 		}
diff --git a/Engine/Plugins/Runtime/WaveTable/Source/WaveTable/Private/WaveTableSampler.cpp b/Engine/Plugins/Runtime/WaveTable/Source/WaveTable/Private/WaveTableSampler.cpp
--- a/Engine/Plugins/Runtime/WaveTable/Source/WaveTable/Private/WaveTableSampler.cpp
+++ b/Engine/Plugins/Runtime/WaveTable/Source/WaveTable/Private/WaveTableSampler.cpp
@@ -23,8 +23,6 @@ namespace WaveTable
 			return FVector2f(2.0f, NextValue - LastValue).GetSafeNormal();
 		};
 
-		using FInlineInterpolator = TFunctionRef<void(TArrayView<const float> InTableView, TArrayView<float> OutSamplesView)>;
-
 		auto CubicIndexInterpolator = [](TArrayView<const float> InTableView, TArrayView<float> InOutIndicesToValues)
 		{
 			const int32 NumTableSamples = InTableView.Num();
@@ -84,6 +82,49 @@ namespace WaveTable
 	{
 	}
 
+	void FWaveTableSampler::Interpolate(TArrayView<const float> InTableView, TArrayView<float> InOutIndexToSamplesView, EInterpolationMode InterpMode)
+	{
+		using namespace UtilitiesPrivate;
+
+		switch (InterpMode)
+		{
+			case EInterpolationMode::None:
+			{
+				StepIndexInterpolator(InTableView, InOutIndexToSamplesView);
+			}
+			break;
+
+			case EInterpolationMode::Linear:
+			{
+				LinearIndexInterpolator(InTableView, InOutIndexToSamplesView);
+			}
+			break;
+
+			case EInterpolationMode::Cubic:
+			{
+				CubicIndexInterpolator(InTableView, InOutIndexToSamplesView);
+			}
+			break;
+
+			default:
+			{
+				static_assert(static_cast<int32>(EInterpolationMode::COUNT) == 3, "Possible missing switch coverage for EInterpolationMode");
+				checkNoEntry();
+			}
+			break;
+		}
+	}
+
+	void FWaveTableSampler::Reset()
+	{
+		LastIndex = 0.0f;
+	}
+
+	const FWaveTableSampler::FSettings& FWaveTableSampler::GetSettings() const
+	{
+		return Settings;
+	}
+
 	void FWaveTableSampler::SetFreq(float InFreq)
 	{
 		Settings.Freq = InFreq;
@@ -99,117 +140,200 @@ namespace WaveTable
 		Settings.Phase = InPhase;
 	}
 
-	void FWaveTableSampler::Process(TArrayView<const float> InTableView, TArrayView<float> OutSamplesView)
+	float FWaveTableSampler::Process(TArrayView<const float> InTableView, float& OutSample, ESingleSampleMode InMode)
 	{
-		Process(InTableView, { }, { }, OutSamplesView);
-	}
+		TArrayView<float> OutSamplesView { &OutSample, 1 };
+		float Index = Process(InTableView, { }, { }, { }, OutSamplesView);
 
-	void FWaveTableSampler::Process(TArrayView<const float> InTableView, TArrayView<const float> InFreqModulator, TArrayView<const float> InPhaseModulator, TArrayView<float> OutSamplesView)
-	{
-		using namespace UtilitiesPrivate;
+		// Sampler functions by default cyclically and interpolates against initial value so
+		// early out as following index check and potential interp re-evaluation isn't necessary.
+		if (InMode == ESingleSampleMode::Loop)
+		{
+			return Index;
+		}
 
-		if (!OutSamplesView.IsEmpty())
+		// If interpolating between last value and next, check
+		// to see if interpolation needs to be re-evaluated.
+		// This check is expensive in the context of buffer
+		// eval and thus only supported in "single-sample" mode.
+		const bool bIsFinalInterp = Index > InTableView.Num() - 1;
+		const bool bHasCompleted = Index == 0 && (Settings.Phase > 0.0f || LastIndex > 0);
+		if (bIsFinalInterp || bHasCompleted)
 		{
-			if (InTableView.IsEmpty())
+			switch (InMode)
 			{
-				Audio::ArraySetToConstantInplace(OutSamplesView, 0.0f);
-			}
-			else
-			{
-				ComputeIndexFrequency(InTableView, InFreqModulator, OutSamplesView);
-				ComputeIndexPhase(InTableView, InPhaseModulator, OutSamplesView);
-
-				switch (Settings.InterpolationMode)
+				case ESingleSampleMode::Hold:
 				{
-					case EInterpolationMode::None:
-					{
-						StepIndexInterpolator(InTableView, OutSamplesView);
-					}
-					break;
+					OutSample = InTableView.Last();
 
-					case EInterpolationMode::Linear:
-					{
-						LinearIndexInterpolator(InTableView, OutSamplesView);
-					}
-					break;
+				}
+				break;
 
-					case EInterpolationMode::Cubic:
+				case ESingleSampleMode::Zero:
+				case ESingleSampleMode::Unit:
+				{
+					const float ModeValue = static_cast<float>(InMode);
+					if (bIsFinalInterp)
 					{
-						CubicIndexInterpolator(InTableView, OutSamplesView);
+						TArray<float> FinalInterp = { InTableView.Last(), ModeValue };
+						OutSample = FMath::Frac(Index);
+						Interpolate(FinalInterp, OutSamplesView, Settings.InterpolationMode);
 					}
-					break;
-
-					default:
+					else
 					{
-						static_assert(static_cast<int32>(EInterpolationMode::COUNT) == 3, "Possible missing switch coverage for EInterpolationMode");
-						checkNoEntry();
+						OutSample = ModeValue;
 					}
-					break;
 				}
+				break;
+
+				case ESingleSampleMode::Loop:
+				{
+					checkNoEntry(); // Should be handled above
+				}
+				break;
+
+				default:
+				break;
+			}
+
+			if (!FMath::IsNearlyEqual(Settings.Amplitude, 1.0f))
+			{
+				OutSample *= Settings.Amplitude;
 			}
 
 			if (!FMath::IsNearlyZero(Settings.Offset))
 			{
-				Audio::ArrayAddConstantInplace(OutSamplesView, Settings.Offset);
+				OutSample += Settings.Offset;
+			}
+		}
+
+		return Index;
+	}
+
+
+	float FWaveTableSampler::Process(TArrayView<const float> InTableView, TArrayView<float> OutSamplesView)
+	{
+		return Process(InTableView, { }, { }, { }, OutSamplesView);
+	}
+
+	float FWaveTableSampler::Process(TArrayView<const float> InTableView, TArrayView<const float> InFreqModulator, TArrayView<const float> InPhaseModulator, TArrayView<const float> InSyncTriggers, TArrayView<float> OutSamplesView)
+	{
+		float CurViewLastIndex = 0.0f;
+
+		if (!OutSamplesView.IsEmpty())
+		{
+			if (InTableView.IsEmpty())
+			{
+				Audio::ArraySetToConstantInplace(OutSamplesView, 0.0f);
+			}
+			else
+			{
+				// Compute index frequency & phase and stuff into samples
+				// view. Interpolator will convert to associated values.
+				ComputeIndexFrequency(InTableView, InFreqModulator, InSyncTriggers, OutSamplesView);
+				ComputeIndexPhase(InTableView, InPhaseModulator, OutSamplesView);
+
+				// Capture last index and return in case caller is interested in progress/phase through table (ex. for enveloping).
+				CurViewLastIndex = OutSamplesView.Last();
+
+				Interpolate(InTableView, OutSamplesView, Settings.InterpolationMode);
 			}
 
 			if (!FMath::IsNearlyEqual(Settings.Amplitude, 1.0f))
 			{
 				Audio::ArrayMultiplyByConstantInPlace(OutSamplesView, Settings.Amplitude);
 			}
+
+			if (!FMath::IsNearlyZero(Settings.Offset))
+			{
+				Audio::ArrayAddConstantInplace(OutSamplesView, Settings.Offset);
+			}
 		}
+
+		return CurViewLastIndex;
 	}
 
-	void FWaveTableSampler::ComputeIndexFrequency(TArrayView<const float> InTableView, TArrayView<const float> InFreqModulator, TArrayView<float> OutSamplesView) const
+	void FWaveTableSampler::ComputeIndexFrequency(TArrayView<const float> InTableView, TArrayView<const float> InFreqModulator, TArrayView<const float> InSyncTriggers, TArrayView<float> OutIndicesView)
 	{
-		const float Delta = 1.0f / OutSamplesView.Num();
-		if (InFreqModulator.Num() == OutSamplesView.Num())
+		float NextIndex = 0.0f;
+
+		Audio::ArraySetToConstantInplace(OutIndicesView, Settings.Freq * InTableView.Num() / OutIndicesView.Num());
+
+		if (InFreqModulator.Num() == OutIndicesView.Num())
 		{
-			FMemory::Memcpy(OutSamplesView.GetData(), InFreqModulator.GetData(), sizeof(float) * OutSamplesView.Num());
-			Audio::ArrayMultiplyByConstantInPlace(OutSamplesView, Settings.Freq * Delta);
+			Audio::ArrayMultiplyInPlace(InFreqModulator, OutIndicesView);
+		}
+
+		if (InSyncTriggers.Num() == OutIndicesView.Num())
+		{
+			float* OutSamples = OutIndicesView.GetData();
+			float SyncIndex = 0;
+			const float* SyncTrigData = InSyncTriggers.GetData();
+			for (int32 i = 0; i < OutIndicesView.Num(); ++i)
+			{
+				SyncIndex *= -1.0f * SyncTrigData[i] + 1.0f;
+				OutSamples[i] *= SyncIndex;
+				OutSamples[i] += LastIndex;
+				SyncIndex++;
+			}
+
+			// Only adjust last index (indicating an inter-block reset)
+			// if the final index is less than the buffer size
+			if (SyncIndex < InSyncTriggers.Num())
+			{
+				LastIndex -= SyncIndex;
+			}
 		}
 		else
 		{
-			checkf(InFreqModulator.IsEmpty(), TEXT("FrequencyModulator view should be the same size as the sample view or not supplied (size of 0)."));
-			Audio::ArraySetToConstantInplace(OutSamplesView, Settings.Freq * Delta);
+			checkf(InSyncTriggers.IsEmpty(), TEXT("SyncTriggers view should be the same size as the sample view or not supplied (size of 0)."));
+
+			float* OutSamples = OutIndicesView.GetData();
+			const float* SyncTrigData = InSyncTriggers.GetData();
+			for (int32 i = 0; i < OutIndicesView.Num(); ++i)
+			{
+				OutSamples[i] *= i;
+				OutSamples[i] += LastIndex;
+			}
 		}
 
-		float* OutSamples = OutSamplesView.GetData();
-		for (int32 i = 0; i < OutSamplesView.Num(); ++i)
+		if (InFreqModulator.Num() == OutIndicesView.Num())
+		{
+			NextIndex = (InFreqModulator.Last() * Settings.Freq * InTableView.Num()) + LastIndex;
+		}
+		else
 		{
-			OutSamples[i] *= i;
+			checkf(InFreqModulator.IsEmpty(), TEXT("FreqModulator view should be the same size as the sample view or not supplied (size of 0)."));
+			NextIndex = (Settings.Freq * InTableView.Num()) + LastIndex;
 		}
+
+		LastIndex = FMath::Frac(NextIndex / InTableView.Num()) * InTableView.Num();
 	}
 
-	void FWaveTableSampler::ComputeIndexPhase(TArrayView<const float> InTableView, TArrayView<const float> InPhaseModulator, TArrayView<float> OutSamplesView)
+	void FWaveTableSampler::ComputeIndexPhase(TArrayView<const float> InTableView, TArrayView<const float> InPhaseModulator, TArrayView<float> OutIndicesView)
 	{
-		const float DeltaPhase = FMath::Frac(LastIndex / InTableView.Num());
-		if (InPhaseModulator.Num() == OutSamplesView.Num())
+		check(!OutIndicesView.IsEmpty());
+
+		if (InPhaseModulator.Num() == OutIndicesView.Num())
 		{
-			PhaseModScratch.SetNum(OutSamplesView.Num());
+			PhaseModScratch.SetNum(OutIndicesView.Num());
 			Audio::ArraySetToConstantInplace(PhaseModScratch, Settings.Phase);
 			Audio::ArrayAddInPlace(InPhaseModulator, PhaseModScratch);
 
 			const float* PhaseModData = PhaseModScratch.GetData();
-			float* OutSamples = OutSamplesView.GetData();
-			for (int32 i = 0; i < OutSamplesView.Num(); ++i)
+			float* OutSamples = OutIndicesView.GetData();
+			for (int32 i = 0; i < OutIndicesView.Num(); ++i)
 			{
-				const float Phase = FMath::Frac(DeltaPhase + PhaseModData[i]);
-				const float Offset = OutSamples[i] + Phase;
-				OutSamples[i] = Offset * InTableView.Num();
+				const float PhaseIndexOffset = FMath::Frac(PhaseModData[i]) * InTableView.Num();
+				OutSamples[i] += PhaseIndexOffset;
 			}
-
-			LastIndex = OutSamplesView.Last() - (PhaseModScratch.Last() * InTableView.Num());
 		}
 		else
 		{
 			checkf(InPhaseModulator.IsEmpty(), TEXT("PhaseModulator view should be the same size as the sample view or not supplied (size of 0)."));
 
-			const float Phase = FMath::Frac(DeltaPhase + Settings.Phase);
-			Audio::ArrayAddConstantInplace(OutSamplesView, Phase);
-			Audio::ArrayMultiplyByConstantInPlace(OutSamplesView, InTableView.Num());
-
-			LastIndex = OutSamplesView.Last() - (Settings.Phase * InTableView.Num());
+			const float PhaseIndexOffset = FMath::Frac(Settings.Phase) * InTableView.Num();
+			Audio::ArrayAddConstantInplace(OutIndicesView, PhaseIndexOffset);
 		}
 	}
 } // namespace WaveTable
diff --git a/Engine/Plugins/Runtime/WaveTable/Source/WaveTable/Private/WaveTableSettings.cpp b/Engine/Plugins/Runtime/WaveTable/Source/WaveTable/Private/WaveTableSettings.cpp
--- a/Engine/Plugins/Runtime/WaveTable/Source/WaveTable/Private/WaveTableSettings.cpp
+++ b/Engine/Plugins/Runtime/WaveTable/Source/WaveTable/Private/WaveTableSettings.cpp
@@ -4,31 +4,33 @@
 
 namespace WaveTable
 {
-	int32 ResolutionToInt32(EWaveTableResolution InResolution)
+	int32 ResolutionToInt32(EWaveTableResolution InWaveTableResolution)
 	{
-		if (InResolution == EWaveTableResolution::None)
+		if (InWaveTableResolution == EWaveTableResolution::None)
 		{
 			return 0;
 		}
 
-		return 1 << static_cast<int32>(InResolution);
+		return 1 << static_cast<int32>(InWaveTableResolution);
 	}
 
-	int32 ResolutionToInt32(EWaveTableCurve InCurve, EWaveTableResolution FileResolution)
+	int32 ResolutionToInt32(EWaveTableResolution InWaveTableResolution, EWaveTableCurve InCurve)
 	{
-		EWaveTableResolution Resolution = EWaveTableResolution::None;
+		if (InWaveTableResolution != EWaveTableResolution::None)
+		{
+			return ResolutionToInt32(InWaveTableResolution);
+		}
+
 		switch (InCurve)
 		{
-			// File uses the incoming resolution
+			// File always uses the incoming resolution, even if 'None'
 			case EWaveTableCurve::File:
-			{
-				Resolution = FileResolution;
-			}
 			break;
 
 			case EWaveTableCurve::Linear:
+			case EWaveTableCurve::Linear_Inv:
 			{
-				Resolution = EWaveTableResolution::Res_16;
+				InWaveTableResolution = EWaveTableResolution::Res_8;
 			}
 			break;
 
@@ -36,25 +38,26 @@ namespace WaveTable
 			case EWaveTableCurve::Exp_Inverse:
 			case EWaveTableCurve::Log:
 			{
-				Resolution = EWaveTableResolution::Res_256;
+				InWaveTableResolution = EWaveTableResolution::Res_256;
 			}
 			break;
 
 			case EWaveTableCurve::Sin:
 			case EWaveTableCurve::SCurve:
+			case EWaveTableCurve::Sin_Full:
 			{
-				Resolution = EWaveTableResolution::Res_64;
+				InWaveTableResolution = EWaveTableResolution::Res_64;
 			}
 			break;
 
 
 			default:
 			{
-				Resolution = EWaveTableResolution::Res_128;
+				InWaveTableResolution = EWaveTableResolution::Res_128;
 			}
 			break;
 		}
 
-		return ResolutionToInt32(Resolution);
+		return ResolutionToInt32(InWaveTableResolution);
 	}
 } // namespace WaveTable
diff --git a/Engine/Plugins/Runtime/WaveTable/Source/WaveTable/Private/WaveTableTransform.cpp b/Engine/Plugins/Runtime/WaveTable/Source/WaveTable/Private/WaveTableTransform.cpp
new file mode 100644
--- /dev/null
+++ b/Engine/Plugins/Runtime/WaveTable/Source/WaveTable/Private/WaveTableTransform.cpp
@@ -0,0 +1,345 @@
+// Copyright Epic Games, Inc. All Rights Reserved.
+#include "WaveTableTransform.h"
+
+#include "DSP/FloatArrayMath.h"
+#include "WaveTableImporter.h"
+#include "WaveTableSampler.h"
+
+
+void FWaveTableTransform::Apply(float& InOutValue, bool bInBipolar) const
+{
+	if (WaveTable.IsEmpty())
+	{
+		TArrayView<float> ValueView = { &InOutValue, 1 };
+		bInBipolar ? SampleCurveBipolar(ValueView) : SampleCurveUnipolar(ValueView);
+	}
+	else
+	{
+		InOutValue *= WaveTable.Num(); // Convert to float index
+		TArrayView<float> ValueView = { &InOutValue, 1 };
+		WaveTable::FWaveTableSampler::Interpolate(WaveTable, ValueView);
+	}
+
+}
+
+void FWaveTableTransform::BuildWaveTable(TArrayView<float> InOutTable, bool bInBipolar) const
+{
+	const float Delta = 1.0f / InOutTable.Num();
+	for (int32 i = 0; i < InOutTable.Num(); ++i)
+	{
+		InOutTable[i] = i * Delta;
+	}
+
+	bInBipolar ? SampleCurveBipolar(InOutTable) : SampleCurveUnipolar(InOutTable);
+}
+
+void FWaveTableTransform::SampleCurveBipolar(TArrayView<float> InOutTable) const
+{
+	// Clamp the input
+	Audio::ArrayClampInPlace(InOutTable, 0.0f, 1.0f);
+
+	float* ValueData = InOutTable.GetData();
+	switch (Curve)
+	{
+		case EWaveTableCurve::Custom:
+		{
+			for (int32 i = 0; i < InOutTable.Num(); ++i)
+			{
+				ValueData[i] = CurveCustom.Eval(ValueData[i]);
+			}
+			break;
+		}
+
+		case EWaveTableCurve::Shared:
+		{
+			if (CurveShared)
+			{
+				for (int32 i = 0; i < InOutTable.Num(); ++i)
+				{
+					ValueData[i] = CurveShared->FloatCurve.Eval(ValueData[i]);
+				}
+			}
+			break;
+		}
+
+		case EWaveTableCurve::Linear_Inv:
+		{
+			for (int32 i = 0; i < InOutTable.Num(); ++i)
+			{
+				ValueData[i] = FMath::Lerp(1.0f, -1.0f, ValueData[i]);
+			}
+		}
+		break;
+
+		case EWaveTableCurve::Linear:
+		{
+			for (int32 i = 0; i < InOutTable.Num(); ++i)
+			{
+				ValueData[i] = FMath::Lerp(-1.0f, 1.0f, ValueData[i]);
+			}
+		}
+		break;
+
+		case EWaveTableCurve::Exp:
+		{
+			// Alpha is limited to between 0.0f and 1.0f and ExponentialScalar
+			// between 0 and 10 to keep ValueData "sane" and avoid float boundary.
+			for (int32 i = 0; i < InOutTable.Num(); ++i)
+			{
+				ValueData[i] *= (FMath::Pow(10.0f, Scalar * (ValueData[i] - 1.0f)));
+				ValueData[i] -= 0.5f;
+			}
+		}
+		break;
+
+		case EWaveTableCurve::Exp_Inverse:
+		{
+			// Alpha is limited to between 0.0f and 1.0f and ExponentialScalar
+			// between 0 and 10 to keep ValueData "sane" and avoid float boundary.
+			for (int32 i = 0; i < InOutTable.Num(); ++i)
+			{
+				ValueData[i] = ((ValueData[i] - 1.0f) * FMath::Pow(10.0f, -1.0f * Scalar * ValueData[i])) + 1.0f;
+				ValueData[i] -= 0.5f;
+			}
+		}
+		break;
+
+		case EWaveTableCurve::Log:
+		{
+			for (int32 i = 0; i < InOutTable.Num(); ++i)
+			{
+				ValueData[i] = (Scalar * FMath::LogX(10.0f, ValueData[i])) + 0.5f;
+			}
+		}
+		break;
+
+		case EWaveTableCurve::Sin:
+		{
+			for (int32 i = 0; i < InOutTable.Num(); ++i)
+			{
+				ValueData[i] = FMath::Sin(HALF_PI * ValueData[i]) - 0.5f;
+			}
+		}
+		break;
+
+		case EWaveTableCurve::SCurve:
+		{
+			for (int32 i = 0; i < InOutTable.Num(); ++i)
+			{
+				ValueData[i] = 0.5f * FMath::Sin(((PI * ValueData[i]) - HALF_PI));
+			}
+		}
+		break;
+
+		case EWaveTableCurve::Sin_Full:
+		{
+			for (int32 i = 0; i < InOutTable.Num(); ++i)
+			{
+				ValueData[i] = FMath::Sin(2.0f * PI * ValueData[i]);
+			}
+		}
+		break;
+
+		case EWaveTableCurve::File:
+		{
+			if (WaveTable.IsEmpty())
+			{
+				Audio::ArraySetToConstantInplace(InOutTable, 0.0f);
+			}
+			else
+			{
+				for (int32 i = 0; i < InOutTable.Num(); ++i)
+				{
+					const int32 Index = FMath::TruncToInt32(WaveTable.Num() * ValueData[i]);
+					ValueData[i] = WaveTable[Index % WaveTable.Num()];
+				}
+			}
+		}
+		break;
+
+		default:
+		{
+			static_assert(static_cast<int32>(EWaveTableCurve::Count) == 11, "Possible missing case coverage for output curve.");
+		}
+		break;
+	}
+
+	Audio::ArrayClampInPlace(InOutTable, -1.0f, 1.0f);
+}
+
+void FWaveTableTransform::SampleCurveUnipolar(TArrayView<float> InOutTable) const
+{
+	// Clamp the input
+	Audio::ArrayClampInPlace(InOutTable, 0.0f, 1.0f);
+
+	float* ValueData = InOutTable.GetData();
+	switch (Curve)
+	{
+		case EWaveTableCurve::Custom:
+		{
+			for (int32 i = 0; i < InOutTable.Num(); ++i)
+			{
+				ValueData[i] = CurveCustom.Eval(ValueData[i]);
+			}
+			break;
+		}
+
+		case EWaveTableCurve::Shared:
+		{
+			if (CurveShared)
+			{
+				for (int32 i = 0; i < InOutTable.Num(); ++i)
+				{
+					ValueData[i] = CurveShared->FloatCurve.Eval(ValueData[i]);
+				}
+			}
+			break;
+		}
+
+		case EWaveTableCurve::Linear_Inv:
+		{
+			for (int32 i = 0; i < InOutTable.Num(); ++i)
+			{
+				ValueData[i] = 1.0f - InOutTable[i];
+			}
+		}
+		break;
+
+		case EWaveTableCurve::Linear:
+		{
+			// Do nothing, as output values == clamped input values
+		}
+		break;
+	
+		case EWaveTableCurve::Exp:
+		{
+			// Alpha is limited to between 0.0f and 1.0f and ExponentialScalar
+			// between 0 and 10 to keep ValueData[i]s "sane" and avoid float boundary.
+			for (int32 i = 0; i < InOutTable.Num(); ++i)
+			{
+				ValueData[i] *= (FMath::Pow(10.0f, Scalar * (ValueData[i] - 1.0f)));
+			}
+		}
+		break;
+
+		case EWaveTableCurve::Exp_Inverse:
+		{
+			// Alpha is limited to between 0.0f and 1.0f and ExponentialScalar
+			// between 0 and 10 to keep ValueData[i]s "sane" and avoid float boundary.
+			for (int32 i = 0; i < InOutTable.Num(); ++i)
+			{
+				ValueData[i] = ((ValueData[i] - 1.0f) * FMath::Pow(10.0f, -1.0f * Scalar * ValueData[i])) + 1.0f;
+			}
+		}
+		break;
+
+		case EWaveTableCurve::Log:
+		{
+			for (int32 i = 0; i < InOutTable.Num(); ++i)
+			{
+				ValueData[i] = (Scalar * FMath::LogX(10.0f, ValueData[i])) + 1.0f;
+			}
+		}
+		break;
+
+		case EWaveTableCurve::Sin:
+		{
+			for (int32 i = 0; i < InOutTable.Num(); ++i)
+			{
+				ValueData[i] = FMath::Sin(HALF_PI * ValueData[i]);
+			}
+		}
+		break;
+
+		case EWaveTableCurve::SCurve:
+		{
+			for (int32 i = 0; i < InOutTable.Num(); ++i)
+			{
+				ValueData[i] = 0.5f * FMath::Sin(((PI * ValueData[i]) - HALF_PI)) + 0.5f;
+			}
+		}
+		break;
+
+		case EWaveTableCurve::Sin_Full:
+		{
+			for (int32 i = 0; i < InOutTable.Num(); ++i)
+			{
+				ValueData[i] = 0.5f * FMath::Sin(2.0f * PI * ValueData[i]) + 0.5f;
+			}
+		}
+		break;
+
+		case EWaveTableCurve::File:
+		{
+			if (WaveTable.IsEmpty())
+			{
+				Audio::ArraySetToConstantInplace(InOutTable, 0.0f);
+			}
+			else
+			{
+				for (int32 i = 0; i < InOutTable.Num(); ++i)
+				{
+					const int32 Index = FMath::TruncToInt32(WaveTable.Num() * ValueData[i]);
+					ValueData[i] = WaveTable[Index % WaveTable.Num()];
+				}
+			}
+		}
+		break;
+
+		default:
+		{
+			static_assert(static_cast<int32>(EWaveTableCurve::Count) == 11, "Possible missing case coverage for output curve.");
+		}
+		break;
+	}
+
+	Audio::ArrayClampInPlace(InOutTable, 0.0f, 1.0f);
+}
+
+void FWaveTableTransform::CacheCurve()
+{
+	if (Curve == EWaveTableCurve::Shared && CurveShared)
+	{
+		CurveCustom = CurveShared->FloatCurve;
+	}
+
+	// Always clear shared curve regardless to avoid
+	// holding a reference to the UObject when never
+	// expected to be used.
+	CurveShared = nullptr;
+}
+
+#if WITH_EDITOR
+void FWaveTableTransform::CacheWaveTable(EWaveTableResolution InResolution, bool bInBipolar)
+{
+	WaveTable.Empty();
+
+	// If not wavetable, just sample using apply function.
+	switch (Curve)
+	{
+		case EWaveTableCurve::File:
+		{
+			WaveTable::FImporter Importer(WaveTableSettings, InResolution, bInBipolar);
+			Importer.Process(WaveTable);
+		}
+		break;
+
+		case EWaveTableCurve::Custom:
+		case EWaveTableCurve::Shared:
+		case EWaveTableCurve::Exp:
+		case EWaveTableCurve::Exp_Inverse:
+		case EWaveTableCurve::Linear:
+		case EWaveTableCurve::Linear_Inv:
+		case EWaveTableCurve::Log:
+		case EWaveTableCurve::SCurve:
+		case EWaveTableCurve::Sin:
+		case EWaveTableCurve::Sin_Full:
+		default:
+		{
+			WaveTable.AddZeroed(WaveTable::ResolutionToInt32(InResolution));
+			BuildWaveTable(WaveTable, bInBipolar);
+		}
+		break;
+	}
+}
+#endif // WITH_EDITOR
diff --git a/Engine/Plugins/Runtime/WaveTable/Source/WaveTable/Public/WaveTable.h b/Engine/Plugins/Runtime/WaveTable/Source/WaveTable/Public/WaveTable.h
new file mode 100644
--- /dev/null
+++ b/Engine/Plugins/Runtime/WaveTable/Source/WaveTable/Public/WaveTable.h
@@ -0,0 +1,74 @@
+// Copyright Epic Games, Inc. All Rights Reserved.
+#pragma once
+
+#include "Containers/Array.h"
+#include "Containers/ArrayView.h"
+#include "HAL/Platform.h"
+
+
+namespace WaveTable
+{
+	struct WAVETABLE_API FWaveTable
+	{
+	private:
+		TArray<float> Samples;
+
+	public:
+		FWaveTable() = default;
+
+		FWaveTable(const TArray<float>& InSamples)
+			: Samples(InSamples)
+		{
+		}
+
+		FWaveTable(TArray<float>&& InSamples)
+			: Samples(MoveTemp(InSamples))
+		{
+		}
+
+		float* GetData()
+		{
+			return Samples.GetData();
+		}
+
+		const float* GetData() const
+		{
+			return Samples.GetData();
+		}
+
+		TArrayView<float> GetView()
+		{
+			return Samples;
+		}
+
+		TArrayView<const float> GetView() const
+		{
+			return Samples;
+		}
+
+		int32 Num() const
+		{
+			return Samples.Num();
+		}
+
+		void Set(TArray<float>&& InSamples)
+		{
+			Samples = MoveTemp(InSamples);
+		}
+
+		void Set(TArrayView<const float> InSamples)
+		{
+			Samples = InSamples;
+		}
+
+		void SetNum(int32 InSize)
+		{
+			Samples.SetNum(InSize);
+		}
+
+		void Zero()
+		{
+			FMemory::Memset(Samples.GetData(), 0, sizeof(float) * Samples.Num());
+		}
+	};
+} // namespace WaveTable
\ No newline at end of file
diff --git a/Engine/Plugins/Runtime/WaveTable/Source/WaveTable/Public/WaveTableBank.h b/Engine/Plugins/Runtime/WaveTable/Source/WaveTable/Public/WaveTableBank.h
new file mode 100644
--- /dev/null
+++ b/Engine/Plugins/Runtime/WaveTable/Source/WaveTable/Public/WaveTableBank.h
@@ -0,0 +1,86 @@
+// Copyright Epic Games, Inc. All Rights Reserved.
+#pragma once
+
+#include "Algo/Transform.h"
+#include "IAudioProxyInitializer.h"
+#include "UObject/Object.h"
+#include "UObject/ObjectMacros.h"
+#include "UObject/ObjectSaveContext.h"
+#include "WaveTable.h"
+#include "WaveTableTransform.h"
+
+#include "WaveTableBank.generated.h"
+
+
+USTRUCT()
+struct WAVETABLE_API FWaveTableBankEntry
+{
+	GENERATED_BODY()
+
+	UPROPERTY(EditAnywhere, Category = Table)
+	FWaveTableTransform Transform;
+};
+
+UCLASS(config = Engine, editinlinenew, BlueprintType)
+class WAVETABLE_API UWaveTableBank : public UObject, public IAudioProxyDataFactory
+{
+	GENERATED_BODY()
+
+public:
+	// If set to 'None', WaveTable is not cached.  If set to
+	// value, caches a sampled wavetable for all curves.
+	UPROPERTY(EditAnywhere, Category = Options)
+	EWaveTableResolution Resolution;
+
+	// Determines if output from curve/wavetable are to be clamped between [-1.0f, 1.0f]
+	// (i.e. for waveform generation, oscillation, etc.) or between [0.0f, 1.0f]
+	// (i.e. for enveloping) when sampling curve/wavetable
+	UPROPERTY(EditAnywhere, Category = Options)
+	bool bBipolar = true;
+
+	/** Tables within the given bank */
+	UPROPERTY(EditAnywhere, Category = Options)
+	TArray<FWaveTableBankEntry> Entries;
+
+	/* IAudioProxyDataFactory Implementation */
+	virtual TUniquePtr<Audio::IProxyData> CreateNewProxyData(const Audio::FProxyDataInitParams& InitParams) override;
+
+#if WITH_EDITOR
+	virtual void PreSave(FObjectPreSaveContext InSaveContext) override;
+	virtual void PostEditChangeProperty(FPropertyChangedEvent& InPropertyChangedEvent) override;
+	virtual void PostEditChangeChainProperty(FPropertyChangedChainEvent& InPropertyChangedEvent) override;
+#endif // WITH_EDITOR
+};
+
+class WAVETABLE_API FWaveTableBankAssetProxy : public Audio::TProxyData<FWaveTableBankAssetProxy>, public TSharedFromThis<FWaveTableBankAssetProxy, ESPMode::ThreadSafe>
+{
+public:
+	IMPL_AUDIOPROXY_CLASS(FWaveTableBankAssetProxy);
+
+	FWaveTableBankAssetProxy(const FWaveTableBankAssetProxy& InAssetProxy)
+		: WaveTables(InAssetProxy.WaveTables)
+	{
+	}
+
+	FWaveTableBankAssetProxy(const UWaveTableBank& InWaveTableBank)
+	{
+		Algo::Transform(InWaveTableBank.Entries, WaveTables, [](const FWaveTableBankEntry& Entry)
+		{
+			return Entry.Transform.WaveTable;
+		});
+	}
+
+	virtual Audio::IProxyDataPtr Clone() const
+	{
+		return MakeUnique<FWaveTableBankAssetProxy>(*this);
+	}
+
+	virtual const TArray<WaveTable::FWaveTable>& GetWaveTables() const
+	{
+		return WaveTables;
+	}
+
+protected:
+	TArray<WaveTable::FWaveTable> WaveTables;
+};
+using FWaveTableBankAssetProxyPtr = TSharedPtr<FWaveTableBankAssetProxy, ESPMode::ThreadSafe>;
diff --git a/Engine/Plugins/Runtime/WaveTable/Source/WaveTable/Public/WaveTableImporter.h b/Engine/Plugins/Runtime/WaveTable/Source/WaveTable/Public/WaveTableImporter.h
--- a/Engine/Plugins/Runtime/WaveTable/Source/WaveTable/Public/WaveTableImporter.h
+++ b/Engine/Plugins/Runtime/WaveTable/Source/WaveTable/Public/WaveTableImporter.h
@@ -11,13 +11,16 @@ namespace WaveTable
 	class WAVETABLE_API FImporter
 	{
 	public:
-		FImporter(const FWaveTableSettings& InOptions, EWaveTableResolution InResolution);
+		FImporter(const FWaveTableSettings& InSettings, EWaveTableResolution InResolution, bool bInBipolar);
 
 		void Process(TArray<float>& OutWaveTable);
 
 	private:
-		const FWaveTableSettings& Options;
+		const FWaveTableSettings& Settings;
+
 		EWaveTableResolution Resolution;
+		bool bBipolar = false;
+
 		FWaveTableSampler Sampler;
 	};
 } // namespace WaveTable
diff --git a/Engine/Plugins/Runtime/WaveTable/Source/WaveTable/Public/WaveTableSampler.h b/Engine/Plugins/Runtime/WaveTable/Source/WaveTable/Public/WaveTableSampler.h
--- a/Engine/Plugins/Runtime/WaveTable/Source/WaveTable/Public/WaveTableSampler.h
+++ b/Engine/Plugins/Runtime/WaveTable/Source/WaveTable/Public/WaveTableSampler.h
@@ -21,6 +21,22 @@ namespace WaveTable
 			COUNT
 		};
 
+		// Mode of interpolation between last value in table and subsequent input index.
+		enum class ESingleSampleMode : uint8
+		{
+			// Interpolates last value to zero (0.0f) in table if sampled index is beyond last position
+			Zero = 0,
+
+			// Interpolates last value to unit (1.0f) in table if sampled index is beyond last position
+			Unit = 1,
+
+			// Holds last value in table if index is beyond last position
+			Hold,
+
+			// Interpolates last value and first value in table if sampled index is beyond last position
+			Loop,
+		};
+
 		struct WAVETABLE_API FSettings
 		{
 			float Amplitude = 1.0f;
@@ -39,18 +55,26 @@ namespace WaveTable
 		FWaveTableSampler();
 		FWaveTableSampler(FSettings&& InSettings);
 
-		void Process(TArrayView<const float> InTableView, TArrayView<float> OutSamplesView);
-		void Process(TArrayView<const float> InTableView, TArrayView<const float> InFreqModulator, TArrayView<const float> InPhaseModulator, TArrayView<float> OutSamplesView);
+		// Interpolates and converts values in the given table for each provided index in the index-to-samples TArrayView (an array of sub-sample, floating point, indices)
+		static void Interpolate(TArrayView<const float> InTableView, TArrayView<float> InOutIndexToSamplesView, EInterpolationMode InterpMode = EInterpolationMode::Linear);
+
+		float Process(TArrayView<const float> InTableView, float& OutSample, ESingleSampleMode InMode = ESingleSampleMode::Zero);
+		float Process(TArrayView<const float> InTableView, TArrayView<float> OutSamplesView);
+		float Process(TArrayView<const float> InTableView, TArrayView<const float> InFreqModulator, TArrayView<const float> InPhaseModulator, TArrayView<const float> InSyncTriggers, TArrayView<float> OutSamplesView);
+
+		void Reset();
+
+		const FSettings& GetSettings() const;
 
 		void SetInterpolationMode(EInterpolationMode InMode);
 		void SetFreq(float InFreq);
 		void SetPhase(float InPhase);
 
 	private:
-		void ComputeIndexFrequency(TArrayView<const float> InTableView, TArrayView<const float> InFreqModulator, TArrayView<float> OutSamplesView) const;
-		void ComputeIndexPhase(TArrayView<const float> InTableView, TArrayView<const float> InPhaseModulator, TArrayView<float> OutSamplesView);
+		void ComputeIndexFrequency(TArrayView<const float> InTableView, TArrayView<const float> InFreqModulator, TArrayView<const float> InSyncTriggers, TArrayView<float> OutIndicesView);
+		void ComputeIndexPhase(TArrayView<const float> InTableView, TArrayView<const float> InPhaseModulator, TArrayView<float> OutIndicesView);
 
-		float LastIndex = 0.0;
+		float LastIndex = 0.0f;
 
 		TArray<float> PhaseModScratch;
 		FSettings Settings;
diff --git a/Engine/Plugins/Runtime/WaveTable/Source/WaveTable/Public/WaveTableSettings.h b/Engine/Plugins/Runtime/WaveTable/Source/WaveTable/Public/WaveTableSettings.h
--- a/Engine/Plugins/Runtime/WaveTable/Source/WaveTable/Public/WaveTableSettings.h
+++ b/Engine/Plugins/Runtime/WaveTable/Source/WaveTable/Public/WaveTableSettings.h
@@ -14,6 +14,7 @@ UENUM(BlueprintType)
 enum class EWaveTableResolution : uint8
 {
 	None		= 0 UMETA(DisplayName = "None"),
+	Res_8		= 3 UMETA(DisplayName = "8"),
 	Res_16		= 4 UMETA(DisplayName = "16"),
 	Res_32		= 5 UMETA(DisplayName = "32"),
 	Res_64		= 6 UMETA(DisplayName = "64"),
@@ -62,17 +63,12 @@ struct WAVETABLE_API FWaveTableSettings
 	float FadeOut = 0.0f;
 
 	// Whether or not to normalize the WaveTable.
-	UPROPERTY(EditAnywhere, Category = Options)
+	UPROPERTY(EditAnywhere, Category = Options, meta = (DisplayAfter = "FilePath"))
 	bool bNormalize = true;
 
-	// If true, converts to envelope, effectively taking
-	// the absolute value of the given file's waveform.
-	UPROPERTY(EditAnywhere, Category = Options)
-	bool bConvertToEnvelope = false;
-
 	// Whether or not to remove offset from original file
 	// (analogous to "DC offset" in circuit theory).
-	UPROPERTY(EditAnywhere, Category = Options)
+	UPROPERTY(EditAnywhere, Category = Options, meta = (DisplayAfter = "bNormalize"))
 	bool bRemoveOffset = false;
 
 	// SourcePCM Data
@@ -83,12 +79,15 @@ struct WAVETABLE_API FWaveTableSettings
 UENUM(BlueprintType)
 enum class EWaveTableCurve : uint8
 {
-	Linear		UMETA(DisplayName = "Linear"),
+	Linear		UMETA(DisplayName = "Linear (Ramp In)"),
+	Linear_Inv	UMETA(DisplayName = "Linear (Ramp Out)"),
 	Exp			UMETA(DisplayName = "Exponential"),
 	Exp_Inverse UMETA(DisplayName = "Exponential (Inverse)"),
 	Log			UMETA(DisplayName = "Log"),
-	Sin			UMETA(DisplayName = "Sin (Quarter)"),
-	SCurve		UMETA(DisplayName = "Sin (S-Curve)"),
+
+	Sin			UMETA(DisplayName = "Sin (90 deg)"),
+	Sin_Full	UMETA(DisplayName = "Sin (360 deg)"),
+	SCurve		UMETA(DisplayName = "Sin (+/- 90 deg)"),
 
 	// Reference a shared curve asset
 	Shared		UMETA(DisplayName = "Shared"),
@@ -104,6 +103,10 @@ enum class EWaveTableCurve : uint8
 
 namespace WaveTable
 {
-	WAVETABLE_API int32 ResolutionToInt32(EWaveTableResolution InResolution);
-	WAVETABLE_API int32 ResolutionToInt32(EWaveTableCurve InCurve, EWaveTableResolution FileResolution);
+	// Converts WaveTableResolution to integer resolution value
+	WAVETABLE_API int32 ResolutionToInt32(EWaveTableResolution InWaveTableResolution);
+
+	// Converts WaveTableResolution to integer resolution value.  If WaveTableResolution is set to none, uses
+	// default value associated with provided curve.
+	WAVETABLE_API int32 ResolutionToInt32(EWaveTableResolution InWaveTableResolution, EWaveTableCurve InCurve);
 } // namespace WaveTable
\ No newline at end of file
diff --git a/Engine/Plugins/Runtime/WaveTable/Source/WaveTable/Public/WaveTableTransform.h b/Engine/Plugins/Runtime/WaveTable/Source/WaveTable/Public/WaveTableTransform.h
new file mode 100644
--- /dev/null
+++ b/Engine/Plugins/Runtime/WaveTable/Source/WaveTable/Public/WaveTableTransform.h
@@ -0,0 +1,75 @@
+// Copyright Epic Games, Inc. All Rights Reserved.
+#pragma once
+
+#include "Containers/Array.h"
+#include "Containers/ArrayView.h"
+#include "Curves/CurveFloat.h"
+#include "Curves/RichCurve.h"
+#include "HAL/Platform.h"
+#include "Templates/Function.h"
+#include "UObject/SoftObjectPath.h"
+#include "UObject/Object.h"
+#include "UObject/ObjectMacros.h"
+#include "WaveTableSettings.h"
+
+#include "WaveTableTransform.generated.h"
+
+
+USTRUCT(BlueprintType)
+struct WAVETABLE_API FWaveTableTransform
+{
+	GENERATED_USTRUCT_BODY()
+
+	/** The curve to apply when transforming the output. */
+	UPROPERTY(EditAnywhere, Category = Input, BlueprintReadWrite, meta = (DisplayName = "Curve Type"))
+	EWaveTableCurve Curve = EWaveTableCurve::Linear;
+
+	/** When curve set to log, exponential or exponential inverse, value is factor 'b' in following equations with output 'y' and input 'x':
+	 *  Exponential: y = x * 10^-b(1-x)
+	 *  Exponential (Inverse): y = ((x - 1) * 10^(-bx)) + 1
+	 *  Logarithmic: y = b * log(x) + 1
+	 */
+	UPROPERTY(EditAnywhere, Category = Input, BlueprintReadWrite, meta = (DisplayName = "Exponential Scalar", ClampMin = "0.1", ClampMax = "10.0", UIMin = "0.1", UIMax = "10.0"))
+	float Scalar = 2.5;
+
+	/** Custom curve to apply if output curve type is set to 'Custom.' */
+	UPROPERTY()
+	FRichCurve CurveCustom;
+
+	/** Asset curve reference to apply if output curve type is set to 'Shared.' */
+	UPROPERTY(EditAnywhere, meta = (DisplayName = "Asset"), Category = Curve, BlueprintReadWrite)
+	UCurveFloat* CurveShared = nullptr;
+
+	UPROPERTY()
+	TArray<float> WaveTable;
+
+#if WITH_EDITORONLY_DATA
+	UPROPERTY(EditAnywhere, Category = WaveTable)
+	FWaveTableSettings WaveTableSettings;
+#endif // WITH_EDITORONLY_DATA
+
+	void Apply(float& InOutValue, bool bInBipolar = false) const;
+
+	// Builds a single period of a WaveTable using this transform in place of the provided ArrayView.
+	// Size of the WaveTable is dictated by the provided ArrayView's size.
+	void BuildWaveTable(TArrayView<float> InOutTable, bool bInBipolar) const;
+
+#if WITH_EDITOR
+	// No-ops if curve is set to WaveTable.
+	void CacheWaveTable(EWaveTableResolution InResolution, bool bInBipolar);
+#endif // WITH_EDITOR
+
+	/** Caches curve data.  Should a shared curve be selected, curve points are copied locally
+	  * to CurveCustom and Curve type is set accordingly.  Can be used to both avoid keeping a
+	  * uobject reference should the transform be used on a non-game thread and potentially to
+	  * further customize a referenced curve locally.
+	  */
+	void CacheCurve();
+
+private:
+	/** Clamps & applies transform to provided values as bipolar signal*/
+	void SampleCurveBipolar(TArrayView<float> InOutValues) const;
+
+	/** Clamps & applies transform to provided values as unipolar signal */
+	void SampleCurveUnipolar(TArrayView<float> InOutValues) const;
+};
diff --git a/Engine/Plugins/Runtime/WaveTable/Source/WaveTable/WaveTable.Build.cs b/Engine/Plugins/Runtime/WaveTable/Source/WaveTable/WaveTable.Build.cs
--- a/Engine/Plugins/Runtime/WaveTable/Source/WaveTable/WaveTable.Build.cs
+++ b/Engine/Plugins/Runtime/WaveTable/Source/WaveTable/WaveTable.Build.cs
@@ -11,6 +11,7 @@ namespace UnrealBuildTool.Rules
 			PublicDependencyModuleNames.AddRange(
 				new string[]
 				{
+					"AudioExtensions",
 					"Core"
 				}
 			);
@@ -19,6 +20,7 @@ namespace UnrealBuildTool.Rules
 				new string[]
 				{
 					"CoreUObject",
+					"Engine",
 					"SignalProcessing"
 				}
 			);
diff --git a/Engine/Plugins/Runtime/WaveTable/Source/WaveTableEditor/Private/AssetTypeActions_WaveTableBank.cpp b/Engine/Plugins/Runtime/WaveTable/Source/WaveTableEditor/Private/AssetTypeActions_WaveTableBank.cpp
new file mode 100644
--- /dev/null
+++ b/Engine/Plugins/Runtime/WaveTable/Source/WaveTableEditor/Private/AssetTypeActions_WaveTableBank.cpp
@@ -0,0 +1,45 @@
+// Copyright Epic Games, Inc. All Rights Reserved.
+#include "AssetTypeActions_WaveTableBank.h"
+
+#include "Templates/SharedPointer.h"
+#include "WaveTableBank.h"
+#include "WaveTableBankEditor.h"
+
+
+#define LOCTEXT_NAMESPACE "AssetTypeActions"
+
+namespace WaveTable
+{
+	namespace Editor
+	{
+		UClass* FAssetTypeActions_WaveTableBank::GetSupportedClass() const
+		{
+			return UWaveTableBank::StaticClass();
+		}
+
+		const TArray<FText>& FAssetTypeActions_WaveTableBank::GetSubMenus() const
+		{
+			static const TArray<FText> SubMenus
+			{
+				LOCTEXT("AssetWaveTableSubMenu", "WaveTable")
+			};
+
+			return SubMenus;
+		}
+
+		void FAssetTypeActions_WaveTableBank::OpenAssetEditor(const TArray<UObject*>& InObjects, TSharedPtr<IToolkitHost> ToolkitHost)
+		{
+			EToolkitMode::Type Mode = ToolkitHost.IsValid() ? EToolkitMode::WorldCentric : EToolkitMode::Standalone;
+
+			for (UObject* Object : InObjects)
+			{
+				if (UWaveTableBank* Patch = Cast<UWaveTableBank>(Object))
+				{
+					TSharedRef<FBankEditor> Editor = MakeShared<FBankEditor>();
+					Editor->Init(Mode, ToolkitHost, Patch);
+				}
+			}
+		}
+	} // namespace Editor
+} // namespace WaveTable
+#undef LOCTEXT_NAMESPACE
diff --git a/Engine/Plugins/Runtime/WaveTable/Source/WaveTableEditor/Private/AssetTypeActions_WaveTableBank.h b/Engine/Plugins/Runtime/WaveTable/Source/WaveTableEditor/Private/AssetTypeActions_WaveTableBank.h
new file mode 100644
--- /dev/null
+++ b/Engine/Plugins/Runtime/WaveTable/Source/WaveTableEditor/Private/AssetTypeActions_WaveTableBank.h
@@ -0,0 +1,23 @@
+// Copyright Epic Games, Inc. All Rights Reserved.
+#pragma once
+
+#include "AssetTypeActions_Base.h"
+
+namespace WaveTable
+{
+	namespace Editor
+	{
+		class FAssetTypeActions_WaveTableBank : public FAssetTypeActions_Base
+		{
+		public:
+			// IAssetTypeActions Implementation
+			virtual FText GetName() const override { return NSLOCTEXT("AssetTypeActions", "AssetTypeActions_WaveTableBank", "WaveTable Bank"); }
+			virtual FColor GetTypeColor() const override { return FColor(232, 122, 0, 255); }
+			virtual const TArray<FText>& GetSubMenus() const override;
+			virtual UClass* GetSupportedClass() const override;
+			virtual uint32 GetCategories() override { return EAssetTypeCategories::Sounds; }
+
+			virtual void OpenAssetEditor(const TArray<UObject*>& InObjects, TSharedPtr<IToolkitHost> ToolkitHost) override;
+		};
+	} // namespace Editor
+} // namespace WaveTable
diff --git a/Engine/Plugins/Runtime/WaveTable/Source/WaveTableEditor/Private/WaveTableBankEditor.cpp b/Engine/Plugins/Runtime/WaveTable/Source/WaveTableEditor/Private/WaveTableBankEditor.cpp
--- a/Engine/Plugins/Runtime/WaveTable/Source/WaveTableEditor/Private/WaveTableBankEditor.cpp
+++ b/Engine/Plugins/Runtime/WaveTable/Source/WaveTableEditor/Private/WaveTableBankEditor.cpp
@@ -20,8 +20,10 @@
 #include "Tree/SCurveEditorTree.h"
 #include "Tree/ICurveEditorTreeItem.h"
 #include "Tree/SCurveEditorTreePin.h"
-#include "WaveTableSettings.h"
+#include "WaveTableBank.h"
 #include "WaveTableSampler.h"
+#include "WaveTableSettings.h"
+#include "WaveTableTransform.h"
 #include "Widgets/Docking/SDockTab.h"
 #include "Widgets/Input/SButton.h"
 #include "Widgets/Input/SNumericDropDown.h"
@@ -37,21 +39,21 @@ namespace WaveTable
 {
 	namespace Editor
 	{
-		const FName FWaveTableBankEditor::AppIdentifier(TEXT("WaveTableEditorApp"));
-		const FName FWaveTableBankEditor::CurveTabId(TEXT("WaveTableEditor_Curves"));
-		const FName FWaveTableBankEditor::PropertiesTabId(TEXT("WaveTableEditor_Properties"));
+		const FName FBankEditorBase::AppIdentifier(TEXT("WaveTableEditorApp"));
+		const FName FBankEditorBase::CurveTabId(TEXT("WaveTableEditor_Curves"));
+		const FName FBankEditorBase::PropertiesTabId(TEXT("WaveTableEditor_Properties"));
 
-		FWaveTableBankEditor::FWaveTableBankEditor()
+		FBankEditorBase::FBankEditorBase()
 		{
 		}
 
-		void FWaveTableBankEditor::RegisterTabSpawners(const TSharedRef<FTabManager>& InTabManager)
+		void FBankEditorBase::RegisterTabSpawners(const TSharedRef<FTabManager>& InTabManager)
 		{
 			WorkspaceMenuCategory = InTabManager->AddLocalWorkspaceMenuCategory(LOCTEXT("WorkspaceMenu_WaveTableEditor", "WaveTable Editor"));
 
 			FAssetEditorToolkit::RegisterTabSpawners(InTabManager);
 
-			InTabManager->RegisterTabSpawner(PropertiesTabId, FOnSpawnTab::CreateSP(this, &FWaveTableBankEditor::SpawnTab_Properties))
+			InTabManager->RegisterTabSpawner(PropertiesTabId, FOnSpawnTab::CreateSP(this, &FBankEditorBase::SpawnTab_Properties))
 				.SetDisplayName(LOCTEXT("DetailsTab", "Details"))
 				.SetGroup(WorkspaceMenuCategory.ToSharedRef())
 				.SetIcon(FSlateIcon(FAppStyle::GetAppStyleSetName(), "LevelEditor.Tabs.Details"));
@@ -63,13 +65,13 @@ namespace WaveTable
 				.SetIcon(CurveIcon);
 		}
 
-		void FWaveTableBankEditor::UnregisterTabSpawners(const TSharedRef<FTabManager>& InTabManager)
+		void FBankEditorBase::UnregisterTabSpawners(const TSharedRef<FTabManager>& InTabManager)
 		{
 			InTabManager->UnregisterTabSpawner(PropertiesTabId);
 			InTabManager->UnregisterTabSpawner(CurveTabId);
 		}
 
-		void FWaveTableBankEditor::Init(const EToolkitMode::Type Mode, const TSharedPtr<IToolkitHost>& InitToolkitHost, UObject* InParentObject)
+		void FBankEditorBase::Init(const EToolkitMode::Type Mode, const TSharedPtr<IToolkitHost>& InitToolkitHost, UObject* InParentObject)
 		{
 			check(InParentObject);
 
@@ -154,24 +156,87 @@ namespace WaveTable
 			}
 		}
 
-		void FWaveTableBankEditor::ClearExpressionCurve(int32 InCurveIndex)
+		void FBankEditorBase::ClearExpressionCurve(int32 InTransformIndex)
 		{
-			if (CurveData.IsValidIndex(InCurveIndex))
+			if (CurveData.IsValidIndex(InTransformIndex))
 			{
-				CurveData[InCurveIndex].ExpressionCurve.Reset();
+				CurveData[InTransformIndex].ExpressionCurve.Reset();
 			}
 		}
 
-		bool FWaveTableBankEditor::RequiresNewCurve(int32 InCurveIndex, const FRichCurve& InRichCurve) const
+		void FBankEditorBase::GenerateExpressionCurve(WaveTable::Editor::FBankEditorBase::FCurveData& OutCurveData, int32 InTransformIndex, EWaveTableCurveSource InSource, bool bInIsUnset)
 		{
-			const FCurveModelID CurveModelID = CurveData[InCurveIndex].ModelID;
+			TSharedPtr<FRichCurve>& Curve = OutCurveData.ExpressionCurve;
+			if (!Curve.IsValid())
+			{
+				Curve = MakeShared<FRichCurve>();
+			}
+
+			const bool bIsBipolar = GetBankIsBipolar();
+
+			if (!GetIsPropertyEditorDisabled())
+			{
+				Curve->Reset();
+
+				if (const FWaveTableTransform* Transform = GetTransform(InTransformIndex))
+				{
+					// Need at least first and last point to generate curve
+					const EWaveTableResolution BankResolution = GetBankResolution();
+					const int32 PointCount = FMath::Max(WaveTable::ResolutionToInt32(BankResolution, Transform->Curve), 2);
+
+					TArray<FRichCurveKey> NewKeys;
+
+					if (BankResolution == EWaveTableResolution::None)
+					{
+						// Terminate curves at last constant point for non-wavetables
+						TArray<float> KeyTable;
+						KeyTable.AddZeroed(PointCount);
+						Transform->BuildWaveTable(KeyTable, bIsBipolar);
+
+						const float Delta = 1.0 / KeyTable.Num();
+						for (int32 i = 0; i < KeyTable.Num(); ++i)
+						{
+							NewKeys.Add({ Delta * i, KeyTable[i] });
+						}
+
+						float TermPoint = 1.0f;
+						Transform->Apply(TermPoint, bIsBipolar);
+						NewKeys.Add({ 1.0f, TermPoint });
+					}
+					else
+					{
+						// Terminate curves at initial point for wavetables
+						TArray<float> KeyTable;
+						KeyTable.AddZeroed(PointCount);
+						Transform->BuildWaveTable(KeyTable, bIsBipolar);
+
+						const float Delta = 1.0 / KeyTable.Num();
+						for (int32 i = 0; i < KeyTable.Num(); ++i)
+						{
+							NewKeys.Add({ Delta * i, KeyTable[i] });
+						}
+
+						NewKeys.Add({ 1.0f, KeyTable[0] });
+					}
+
+					Curve->SetKeys(NewKeys);
+				}
+			}
+
+			const EWaveTableCurveSource Source = bInIsUnset ? EWaveTableCurveSource::Unset : EWaveTableCurveSource::Expression;
+			SetCurve(InTransformIndex, *Curve.Get(), Source);
+		}
+
+		bool FBankEditorBase::RequiresNewCurve(int32 InTransformIndex, const FRichCurve& InRichCurve) const
+		{
+			const FCurveModelID CurveModelID = CurveData[InTransformIndex].ModelID;
 			const TUniquePtr<FCurveModel>* CurveModel = CurveEditor->GetCurves().Find(CurveModelID);
 			if (!CurveModel || !CurveModel->IsValid())
 			{
 				return true;
 			}
 
-			FWaveTableCurveModelBase* PatchCurveModel = static_cast<FWaveTableCurveModelBase*>(CurveModel->Get());
+			FWaveTableCurveModel* PatchCurveModel = static_cast<FWaveTableCurveModel*>(CurveModel->Get());
 			check(PatchCurveModel);
 			if (&PatchCurveModel->GetRichCurve() != &InRichCurve)
 			{
@@ -181,62 +246,68 @@ namespace WaveTable
 			return false;
 		}
 
-		void FWaveTableBankEditor::SetCurve(int32 InCurveIndex, FRichCurve& InRichCurve, EWaveTableCurveSource InSource)
+		void FBankEditorBase::SetCurve(int32 InTransformIndex, FRichCurve& InRichCurve, EWaveTableCurveSource InSource)
 		{
 			check(CurveEditor.IsValid());
 
-			if (!ensure(CurveData.IsValidIndex(InCurveIndex)))
+			if (!ensure(CurveData.IsValidIndex(InTransformIndex)))
 			{
 				return;
 			}
 
-			FCurveData& CurveDataEntry = CurveData[InCurveIndex];
+			FWaveTableTransform* Transform = GetTransform(InTransformIndex);
+			if (!ensure(Transform))
+			{
+				return;
+			}
 
-			const bool bRequiresNewCurve = RequiresNewCurve(InCurveIndex, InRichCurve);
+			FCurveData& CurveDataEntry = CurveData[InTransformIndex];
+
+			const bool bRequiresNewCurve = RequiresNewCurve(InTransformIndex, InRichCurve);
 			if (bRequiresNewCurve)
 			{
-				TUniquePtr<FWaveTableCurveModelBase> NewCurve = ConstructCurveModel(InRichCurve, GetEditingObject(), InSource);
-				NewCurve->Refresh(InCurveIndex);
+				TUniquePtr<FWaveTableCurveModel> NewCurve = ConstructCurveModel(InRichCurve, GetEditingObject(), InSource);
+				NewCurve->Refresh(*Transform, InTransformIndex);
 				CurveDataEntry.ModelID = CurveEditor->AddCurve(MoveTemp(NewCurve));
 			}
 			else
 			{
 				const TUniquePtr<FCurveModel>& CurveModel = CurveEditor->GetCurves().FindChecked(CurveDataEntry.ModelID);
 				check(CurveModel.Get());
-				static_cast<FWaveTableCurveModelBase*>(CurveModel.Get())->Refresh(InCurveIndex);
+				static_cast<FWaveTableCurveModel*>(CurveModel.Get())->Refresh(*Transform, InTransformIndex);
 			}
 
 			CurveEditor->PinCurve(CurveDataEntry.ModelID);
 		}
 
-		FName FWaveTableBankEditor::GetToolkitFName() const
+		FName FBankEditorBase::GetToolkitFName() const
 		{
 			return FName("WaveTableEditor");
 		}
 
-		FText FWaveTableBankEditor::GetBaseToolkitName() const
+		FText FBankEditorBase::GetBaseToolkitName() const
 		{
 			return LOCTEXT( "AppLabel", "WaveTable Editor" );
 		}
 
-		FString FWaveTableBankEditor::GetWorldCentricTabPrefix() const
+		FString FBankEditorBase::GetWorldCentricTabPrefix() const
 		{
 			return LOCTEXT("WorldCentricTabPrefix", "WaveTable ").ToString();
 		}
 
-		FLinearColor FWaveTableBankEditor::GetWorldCentricTabColorScale() const
+		FLinearColor FBankEditorBase::GetWorldCentricTabColorScale() const
 		{
 			return FLinearColor( 0.0f, 0.0f, 0.2f, 0.5f );
 		}
 
-		EOrientation FWaveTableBankEditor::GetSnapLabelOrientation() const
+		EOrientation FBankEditorBase::GetSnapLabelOrientation() const
 		{
 			return FMultiBoxSettings::UseSmallToolBarIcons.Get()
 				? EOrientation::Orient_Horizontal
 				: EOrientation::Orient_Vertical;
 		}
 
-		void FWaveTableBankEditor::NotifyPostChange(const FPropertyChangedEvent& PropertyChangedEvent, FProperty* PropertyThatChanged)
+		void FBankEditorBase::NotifyPostChange(const FPropertyChangedEvent& PropertyChangedEvent, FProperty* PropertyThatChanged)
 		{
 			if (PropertyChangedEvent.ChangeType != EPropertyChangeType::Interactive)
 			{
@@ -244,7 +315,7 @@ namespace WaveTable
 			}
 		}
 
-		TSharedRef<SDockTab> FWaveTableBankEditor::SpawnTab_Properties(const FSpawnTabArgs& Args)
+		TSharedRef<SDockTab> FBankEditorBase::SpawnTab_Properties(const FSpawnTabArgs& Args)
 		{
 			check(Args.GetTabId() == PropertiesTabId);
 
@@ -255,7 +326,7 @@ namespace WaveTable
 				];
 		}
 
-		TSharedRef<SDockTab> FWaveTableBankEditor::SpawnTab_OutputCurve(const FSpawnTabArgs& Args)
+		TSharedRef<SDockTab> FBankEditorBase::SpawnTab_OutputCurve(const FSpawnTabArgs& Args)
 		{
 			RefreshCurves();
 			CurveEditor->ZoomToFit();
@@ -275,7 +346,7 @@ namespace WaveTable
 				return NewDockTab;
 		}
 
-		void FWaveTableBankEditor::PostUndo(bool bSuccess)
+		void FBankEditorBase::PostUndo(bool bSuccess)
 		{
 			if (bSuccess)
 			{
@@ -283,28 +354,35 @@ namespace WaveTable
 			}
 		}
 
-		void FWaveTableBankEditor::ResetCurves()
+		void FBankEditorBase::ResetCurves()
 		{
 			check(CurveEditor.IsValid());
 
 			CurveEditor->RemoveAllCurves();
 			CurveData.Reset();
-			CurveData.AddDefaulted(GetNumCurves());
+			CurveData.AddDefaulted(GetNumTransforms());
 		}
 
-		void FWaveTableBankEditor::InitCurves()
+		void FBankEditorBase::InitCurves()
 		{
-			for (int32 i = 0; i < GetNumCurves(); ++i)
+			for (int32 i = 0; i < GetNumTransforms(); ++i)
 			{
-				EWaveTableCurve CurveType = GetCurveType(i);
-				switch (CurveType)
+				const FWaveTableTransform* Transform = GetTransform(i);
+				if (!ensure(Transform))
+				{
+					continue;
+				}
+
+				switch (Transform->Curve)
 				{
 					case EWaveTableCurve::Exp:
 					case EWaveTableCurve::Exp_Inverse:
 					case EWaveTableCurve::Linear:
+					case EWaveTableCurve::Linear_Inv:
 					case EWaveTableCurve::Log:
 					case EWaveTableCurve::SCurve:
 					case EWaveTableCurve::Sin:
+					case EWaveTableCurve::Sin_Full:
 					case EWaveTableCurve::File:
 					{
 						if (RequiresNewCurve(i, *CurveData[i].ExpressionCurve.Get()))
@@ -316,7 +394,7 @@ namespace WaveTable
 
 					case EWaveTableCurve::Shared:
 					{
-						if (const UCurveFloat* SharedCurve = GetSharedCurve(i))
+						if (const UCurveFloat* SharedCurve = Transform->CurveShared)
 						{
 							if (RequiresNewCurve(i, SharedCurve->FloatCurve))
 							{
@@ -332,7 +410,7 @@ namespace WaveTable
 
 					case EWaveTableCurve::Custom:
 					{
-						if (RequiresNewCurve(i, GetCustomCurveChecked(i)))
+						if (RequiresNewCurve(i, Transform->CurveCustom))
 						{
 							ResetCurves();
 						}
@@ -341,14 +419,14 @@ namespace WaveTable
 
 					default:
 					{
-						static_assert(static_cast<int32>(EWaveTableCurve::Count) == 9, "Possible missing case coverage for output curve.");
+						static_assert(static_cast<int32>(EWaveTableCurve::Count) == 11, "Possible missing case coverage for output curve.");
 					}
 					break;
 				}
 			}
 		}
 
-		void FWaveTableBankEditor::RefreshCurves()
+		void FBankEditorBase::RefreshCurves()
 		{
 			check(CurveEditor.IsValid());
 
@@ -363,7 +441,7 @@ namespace WaveTable
 				return;
 			}
 
-			const int32 NumCurves = GetNumCurves();
+			const int32 NumCurves = GetNumTransforms();
 			if (NumCurves == CurveData.Num())
 			{
 				InitCurves();
@@ -373,24 +451,30 @@ namespace WaveTable
 				ResetCurves();
 			}
 
-			for (int32 i = 0; i < GetNumCurves(); ++i)
+			for (int32 i = 0; i < GetNumTransforms(); ++i)
 			{
-				const EWaveTableCurve CurveType = GetCurveType(i);
 				if (!CurveData.IsValidIndex(i))
 				{
 					continue;
 				}
-
 				FCurveData& CurveDataEntry = CurveData[i];
 
-				switch (CurveType)
+				FWaveTableTransform* Transform = GetTransform(i);
+				if (!Transform)
+				{
+					continue;
+				}
+
+				switch (Transform->Curve)
 				{
 					case EWaveTableCurve::Exp:
 					case EWaveTableCurve::Exp_Inverse:
 					case EWaveTableCurve::Linear:
+					case EWaveTableCurve::Linear_Inv:
 					case EWaveTableCurve::Log:
 					case EWaveTableCurve::SCurve:
 					case EWaveTableCurve::Sin:
+					case EWaveTableCurve::Sin_Full:
 					case EWaveTableCurve::File:
 					{
 						GenerateExpressionCurve(CurveDataEntry, i, EWaveTableCurveSource::Expression);
@@ -399,7 +483,7 @@ namespace WaveTable
 
 					case EWaveTableCurve::Shared:
 					{
-						if (UCurveFloat* SharedCurve = GetSharedCurve(i))
+						if (UCurveFloat* SharedCurve = Transform->CurveShared)
 						{
 							ClearExpressionCurve(i);
 							SetCurve(i, SharedCurve->FloatCurve, EWaveTableCurveSource::Shared);
@@ -415,7 +499,7 @@ namespace WaveTable
 
 					case EWaveTableCurve::Custom:
 					{
-						FRichCurve& CustomCurve = GetCustomCurveChecked(i);
+						FRichCurve& CustomCurve = Transform->CurveCustom;
 						TrimKeys(CustomCurve);
 						ClearExpressionCurve(i);
 						SetCurve(i, CustomCurve, EWaveTableCurveSource::Custom);
@@ -424,7 +508,7 @@ namespace WaveTable
 
 					default:
 					{
-						static_assert(static_cast<int32>(EWaveTableCurve::Count) == 9, "Possible missing case coverage for output curve.");
+						static_assert(static_cast<int32>(EWaveTableCurve::Count) == 11, "Possible missing case coverage for output curve.");
 					}
 					break;
 				}
@@ -453,7 +537,7 @@ namespace WaveTable
 			}
 		}
 
-		void FWaveTableBankEditor::PostRedo(bool bSuccess)
+		void FBankEditorBase::PostRedo(bool bSuccess)
 		{
 			if (bSuccess)
 			{
@@ -461,7 +545,7 @@ namespace WaveTable
 			}
 		}
 
-		void FWaveTableBankEditor::TrimKeys(FRichCurve& OutCurve)
+		void FBankEditorBase::TrimKeys(FRichCurve& OutCurve)
 		{
 			while (OutCurve.GetNumKeys() > 0 && 0.0f > OutCurve.GetFirstKey().Time)
 			{
@@ -473,6 +557,55 @@ namespace WaveTable
 				OutCurve.DeleteKey(OutCurve.GetLastKeyHandle());
 			}
 		}
+
+		TUniquePtr<WaveTable::Editor::FWaveTableCurveModel> FBankEditor::ConstructCurveModel(FRichCurve& InRichCurve, UObject* InParentObject, EWaveTableCurveSource InSource)
+		{
+			using namespace WaveTable::Editor;
+			return MakeUnique<FWaveTableCurveModel>(InRichCurve, GetEditingObject(), InSource);
+		}
+
+		EWaveTableResolution FBankEditor::GetBankResolution() const
+		{
+			if (UWaveTableBank* Bank = Cast<UWaveTableBank>(GetEditingObject()))
+			{
+				return Bank->Resolution;
+			}
+
+			return EWaveTableResolution::None;
+		}
+
+		bool FBankEditor::GetBankIsBipolar() const
+		{
+			if (UWaveTableBank* Bank = Cast<UWaveTableBank>(GetEditingObject()))
+			{
+				return Bank->bBipolar;
+			}
+
+			return false;
+		}
+
+		int32 FBankEditor::GetNumTransforms() const
+		{
+			if (UWaveTableBank* Bank = Cast<UWaveTableBank>(GetEditingObject()))
+			{
+				return Bank->Entries.Num();
+			}
+
+			return 0;
+		}
+
+		FWaveTableTransform* FBankEditor::GetTransform(int32 InIndex) const
+		{
+			if (UWaveTableBank* Bank = Cast<UWaveTableBank>(GetEditingObject()))
+			{
+				if (Bank->Entries.IsValidIndex(InIndex))
+				{
+					return &Bank->Entries[InIndex].Transform;
+				}
+			}
+
+			return nullptr;
+		}
 	} // namespace Editor
 } // namespace WaveTable
 #undef LOCTEXT_NAMESPACE
diff --git a/Engine/Plugins/Runtime/WaveTable/Source/WaveTableEditor/Private/WaveTableBankFactory.cpp b/Engine/Plugins/Runtime/WaveTable/Source/WaveTableEditor/Private/WaveTableBankFactory.cpp
new file mode 100644
--- /dev/null
+++ b/Engine/Plugins/Runtime/WaveTable/Source/WaveTableEditor/Private/WaveTableBankFactory.cpp
@@ -0,0 +1,18 @@
+// Copyright Epic Games, Inc. All Rights Reserved.
+#include "WaveTableBankFactory.h"
+#include "WaveTableBank.h"
+
+
+UWaveTableBankFactory::UWaveTableBankFactory(const FObjectInitializer& ObjectInitializer)
+	: Super(ObjectInitializer)
+{
+	SupportedClass = UWaveTableBank::StaticClass();
+	bCreateNew = true;
+	bEditorImport = false;
+	bEditAfterNew = true;
+}
+
+UObject* UWaveTableBankFactory::FactoryCreateNew(UClass* Class, UObject* InParent, FName Name, EObjectFlags Flags, UObject* Context, FFeedbackContext* Warn)
+{
+	return NewObject<UWaveTableBank>(InParent, Name, Flags);
+}
diff --git a/Engine/Plugins/Runtime/WaveTable/Source/WaveTableEditor/Private/WaveTableBankFactory.h b/Engine/Plugins/Runtime/WaveTable/Source/WaveTableEditor/Private/WaveTableBankFactory.h
new file mode 100644
--- /dev/null
+++ b/Engine/Plugins/Runtime/WaveTable/Source/WaveTableEditor/Private/WaveTableBankFactory.h
@@ -0,0 +1,19 @@
+// Copyright Epic Games, Inc. All Rights Reserved.
+#pragma once
+
+#include "CoreMinimal.h"
+#include "UObject/ObjectMacros.h"
+#include "Factories/Factory.h"
+
+#include "WaveTableBankFactory.generated.h"
+
+
+UCLASS(hidecategories=Object, MinimalAPI)
+class UWaveTableBankFactory : public UFactory
+{
+	GENERATED_UCLASS_BODY()
+
+	// Begin UFactory Interface
+	virtual UObject* FactoryCreateNew(UClass* Class,UObject* InParent,FName Name,EObjectFlags Flags,UObject* Context,FFeedbackContext* Warn) override;
+	// Begin UFactory Interface	
+};
diff --git a/Engine/Plugins/Runtime/WaveTable/Source/WaveTableEditor/Private/WaveTableCurveEditorViewStacked.cpp b/Engine/Plugins/Runtime/WaveTable/Source/WaveTableEditor/Private/WaveTableCurveEditorViewStacked.cpp
--- a/Engine/Plugins/Runtime/WaveTable/Source/WaveTableEditor/Private/WaveTableCurveEditorViewStacked.cpp
+++ b/Engine/Plugins/Runtime/WaveTable/Source/WaveTableEditor/Private/WaveTableCurveEditorViewStacked.cpp
@@ -10,6 +10,7 @@
 #include "SCurveEditorPanel.h"
 #include "SGraphActionMenu.h"
 #include "Styling/AppStyle.h"
+#include "WaveTableTransform.h"
 #include "Widgets/Text/STextBlock.h"
 
 #define LOCTEXT_NAMESPACE "WaveTableEditor"
@@ -19,7 +20,7 @@ namespace WaveTable
 {
 	namespace Editor
 	{
-		FWaveTableCurveModelBase::FWaveTableCurveModelBase(FRichCurve& InRichCurve, UObject* InOwner, EWaveTableCurveSource InSource)
+		FWaveTableCurveModel::FWaveTableCurveModel(FRichCurve& InRichCurve, UObject* InOwner, EWaveTableCurveSource InSource)
 			: FRichCurveEditorModelRaw(&InRichCurve, InOwner)
 			, ParentObject(InOwner)
 			, Source(InSource)
@@ -27,24 +28,23 @@ namespace WaveTable
 			check(InOwner);
 		}
 
-		void FWaveTableCurveModelBase::Refresh(int32 InCurveIndex)
+		void FWaveTableCurveModel::Refresh(const FWaveTableTransform& InTransform, int32 InCurveIndex)
 		{
 			// Must be set prior to remainder of refresh to avoid
 			// child class implementation acting on incorrect cached index
-			// (Editor May re-purpose models resulting in index change).
+			// (Editor may re-purpose models resulting in index change).
 			CurveIndex = InCurveIndex;
 
 			Color = GetCurveColor();
 
 			FText OutputAxisName;
-			RefreshCurveDescriptorText(ShortDisplayName, InputAxisName, OutputAxisName);
+			RefreshCurveDescriptorText(InTransform, ShortDisplayName, InputAxisName, OutputAxisName);
 
-			ShortDisplayName = GetShortDisplayName();
 			AxesDescriptor = FText::Format(LOCTEXT("WaveTableCurveDisplayTitle_AxesFormat", "X = {0}, Y = {1}"), InputAxisName, OutputAxisName);
 
 			bKeyDrawEnabled = true;
 			IntentionName = TEXT("AudioControlValue");
-			SupportedViews = ViewId;
+			SupportedViews = GetViewId();
 
 			bKeyDrawEnabled = false;
 
@@ -58,15 +58,7 @@ namespace WaveTable
 				UEnum* Enum = StaticEnum<EWaveTableCurveSource>();
 				check(Enum);
 
-				FString EnumValueName = Enum->GetValueAsString(Source);
-				int32 DelimIndex = -1;
-				EnumValueName.FindLastChar(':', DelimIndex);
-				if (DelimIndex > 0 && DelimIndex < EnumValueName.Len() - 1)
-				{
-					EnumValueName.RightChopInline(DelimIndex + 1);
-				}
-
-				FText CurveSourceText = FText::FromString(EnumValueName);
+				LongDisplayName = FText::Format(LOCTEXT("WaveTableCurveDisplayTitle_NameFormat", "{0}: {1}"), FText::AsNumber(CurveIndex), ShortDisplayName);
 
 				if (Source == EWaveTableCurveSource::Custom)
 				{
@@ -74,44 +66,79 @@ namespace WaveTable
 				}
 				else if (Source == EWaveTableCurveSource::Shared)
 				{
-					const UCurveFloat* SharedCurve = GetSharedCurve();
-					check(SharedCurve);
-					CurveSourceText = FText::Format(LOCTEXT("WaveTableCurveDisplayTitle_SharedNameFormat", "Shared, {0}"), FText::FromString(SharedCurve->GetName()));
 					bKeyDrawEnabled = true;
+
+					if (FWaveTableTransform* Transform = GetTransform())
+					{
+						if (const UCurveFloat* SharedCurve = Transform->CurveShared)
+						{
+							const FText CurveSourceText = FText::FromString(SharedCurve->GetName());
+							LongDisplayName = FText::Format(LOCTEXT("WaveTableCurveDisplayTitle_NameSharedFormat", "{0}: {1}, Shared Curve '{2}'"), FText::AsNumber(CurveIndex), ShortDisplayName, CurveSourceText);
+						}
+					}
 				}
 
-				LongDisplayName = FText::Format(LOCTEXT("WaveTableCurveDisplayTitle_NameFormat", "{0}: {1} ({2})"), FText::AsNumber(CurveIndex), ShortDisplayName, CurveSourceText);
 			}
 		}
 
-		FLinearColor FWaveTableCurveModelBase::GetColor() const
+		FLinearColor FWaveTableCurveModel::GetColor() const
 		{
 			return !GetPropertyEditorDisabled() && (Source == EWaveTableCurveSource::Custom)
 				? Color
 				: Color.Desaturate(0.35f);
 		}
 
-		const FText& FWaveTableCurveModelBase::GetAxesDescriptor() const
+		const FText& FWaveTableCurveModel::GetAxesDescriptor() const
 		{
 			return AxesDescriptor;
 		}
 
-		EWaveTableCurveSource FWaveTableCurveModelBase::GetSource() const
+		EWaveTableCurveSource FWaveTableCurveModel::GetSource() const
 		{
 			return Source;
 		}
 
-		const UObject* FWaveTableCurveModelBase::GetParentObject() const
+		const UObject* FWaveTableCurveModel::GetParentObject() const
 		{
 			return ParentObject.Get();
 		}
 
-		bool FWaveTableCurveModelBase::IsReadOnly() const
+		FWaveTableTransform* FWaveTableCurveModel::GetTransform()
+		{
+			return nullptr;
+		}
+
+		bool FWaveTableCurveModel::IsReadOnly() const
 		{
 			return Source != EWaveTableCurveSource::Custom;
 		}
 
-		ECurveEditorViewID FWaveTableCurveModelBase::ViewId = ECurveEditorViewID::Invalid;
+		ECurveEditorViewID FWaveTableCurveModel::WaveTableViewId = ECurveEditorViewID::Invalid;
+
+		void FWaveTableCurveModel::RefreshCurveDescriptorText(const FWaveTableTransform& InTransform, FText& OutShortDisplayName, FText& OutInputAxisName, FText& OutOutputAxisName)
+		{
+			UEnum* Enum = StaticEnum<EWaveTableCurve>();
+			check(Enum);
+
+			Enum->GetDisplayValueAsText(InTransform.Curve, OutShortDisplayName);
+			OutInputAxisName = LOCTEXT("WaveTableCurveModel_DefaultInputAxisName", "Unit Phase (Table Index)");
+			OutOutputAxisName = LOCTEXT("WaveTableCurveModel_DefaultOutputAxisName", "Amplitude");
+		}
+
+		FColor FWaveTableCurveModel::GetCurveColor() const
+		{
+			return FColor { 232, 122, 0, 255 };
+		}
+		
+		bool FWaveTableCurveModel::GetPropertyEditorDisabled() const
+		{
+			return false;
+		}
+
+		FText FWaveTableCurveModel::GetPropertyEditorDisabledText() const
+		{
+			return LOCTEXT("WaveTableCurveModel_Disabled", "Disabled");
+		}
 
 		void SViewStacked::Construct(const FArguments& InArgs, TWeakPtr<FCurveEditor> InCurveEditor)
 		{
@@ -186,7 +213,7 @@ namespace WaveTable
 				FSlateDrawElement::MakeText(OutDrawElements, LabelLayerId + 1, LabelGeometry, Label, FontInfo, DrawEffects, Curve->GetColor());
 
 				// Render axes descriptor
-				FText Descriptor = static_cast<FWaveTableCurveModelBase*>(Curve)->GetAxesDescriptor();
+				FText Descriptor = static_cast<FWaveTableCurveModel*>(Curve)->GetAxesDescriptor();
 
 				const FVector2D DescriptorSize = FontMeasure->Measure(Descriptor, FontInfo);
 
@@ -235,7 +262,7 @@ namespace WaveTable
 				TArray<FText> CurveModelGridLabelsX = MajorGridLabelsX;
 				check(MajorGridLinesX.Num() == CurveModelGridLabelsX.Num());
 
-				if (const FWaveTableCurveModelBase* EditorModel = static_cast<const FWaveTableCurveModelBase*>(DrawInfo.GetCurveModel()))
+				if (const FWaveTableCurveModel* EditorModel = static_cast<const FWaveTableCurveModel*>(DrawInfo.GetCurveModel()))
 				{
 					for (int32 i = 0; i < CurveModelGridLabelsX.Num(); ++i)
 					{
@@ -322,7 +349,7 @@ namespace WaveTable
 
 			FText Label = FText::AsNumber(FMath::Lerp(ValueMin, ValueMax, OffsetAlpha), &DrawInfo.LabelFormat);
 
-			if (const FWaveTableCurveModelBase* CurveModel = static_cast<const FWaveTableCurveModelBase*>(DrawInfo.GetCurveModel()))
+			if (const FWaveTableCurveModel* CurveModel = static_cast<const FWaveTableCurveModel*>(DrawInfo.GetCurveModel()))
 			{
 				FormatOutputLabel(*CurveModel, DrawInfo.LabelFormat, Label);
 			}
diff --git a/Engine/Plugins/Runtime/WaveTable/Source/WaveTableEditor/Private/WaveTableEditorModule.cpp b/Engine/Plugins/Runtime/WaveTable/Source/WaveTableEditor/Private/WaveTableEditorModule.cpp
--- a/Engine/Plugins/Runtime/WaveTable/Source/WaveTableEditor/Private/WaveTableEditorModule.cpp
+++ b/Engine/Plugins/Runtime/WaveTable/Source/WaveTableEditor/Private/WaveTableEditorModule.cpp
@@ -1,19 +1,84 @@
 // Copyright Epic Games, Inc. All Rights Reserved.
 #include "WaveTableEditorModule.h"
 
+#include "AssetTypeActions_WaveTableBank.h"
+#include "ICurveEditorModule.h"
+#include "PropertyEditorDelegates.h"
+#include "PropertyEditorModule.h"
+#include "WaveTableCurveEditorViewStacked.h"
+#include "WaveTableTransformLayout.h"
 
 DEFINE_LOG_CATEGORY(LogWaveTableEditor);
 
+
 namespace WaveTable
 {
 	namespace Editor
 	{
+		namespace ModulePrivate
+		{
+			static const FName AssetToolName { "AssetTools" };
+			static const FName CurveEditorName { "CurveEditor" };
+			static const FName PropertyEditorName { "PropertyEditor" };
+
+			template <typename T>
+			void AddAssetAction(IAssetTools& AssetTools, TArray<TSharedPtr<FAssetTypeActions_Base>>& AssetArray)
+			{
+				TSharedPtr<T> AssetAction = MakeShared<T>();
+				TSharedPtr<FAssetTypeActions_Base> AssetActionBase = StaticCastSharedPtr<FAssetTypeActions_Base>(AssetAction);
+				AssetTools.RegisterAssetTypeActions(AssetAction.ToSharedRef());
+				AssetArray.Add(AssetActionBase);
+			}
+		} // namespace ModulePrivate
+
 		void FModule::StartupModule()
 		{
+			using namespace ModulePrivate;
+
+			FPropertyEditorModule& PropertyModule = FModuleManager::LoadModuleChecked<FPropertyEditorModule>(PropertyEditorName);
+			PropertyModule.RegisterCustomPropertyTypeLayout
+			(
+				"WaveTableTransform",
+				FOnGetPropertyTypeCustomizationInstance::CreateStatic(&FTransformLayoutCustomization::MakeInstance)
+			);
+
+			ICurveEditorModule& CurveEditorModule = FModuleManager::LoadModuleChecked<ICurveEditorModule>(CurveEditorName);
+			FWaveTableCurveModel::WaveTableViewId = CurveEditorModule.RegisterView(FOnCreateCurveEditorView::CreateStatic(
+				[](TWeakPtr<FCurveEditor> WeakCurveEditor) -> TSharedRef<SCurveEditorView>
+				{
+					return SNew(SViewStacked, WeakCurveEditor);
+				}
+			));
+
+			IAssetTools& AssetTools = FModuleManager::LoadModuleChecked<FAssetToolsModule>(AssetToolName).Get();
+			AddAssetAction<FAssetTypeActions_WaveTableBank>(AssetTools, AssetActions);
 		}
 
 		void FModule::ShutdownModule()
 		{
+			using namespace ModulePrivate;
+
+			if (FModuleManager::Get().IsModuleLoaded(PropertyEditorName))
+			{
+				FPropertyEditorModule& PropertyModule = FModuleManager::GetModuleChecked<FPropertyEditorModule>(PropertyEditorName);
+				PropertyModule.UnregisterCustomPropertyTypeLayout("WaveTableTransform");
+			}
+
+// 			if (FModuleManager::Get().IsModuleLoaded(AssetToolName))
+// 			{
+// 				IAssetTools& AssetTools = FModuleManager::GetModuleChecked<FAssetToolsModule>(AssetToolName).Get();
+// 				for (const TSharedPtr<FAssetTypeActions_Base>& ActionPtr : AssetActions)
+// 				{
+// 					AssetTools.UnregisterAssetTypeActions(ActionPtr);
+// 				}
+// 			}
+
+			if (FModuleManager::Get().IsModuleLoaded(CurveEditorName))
+			{
+				ICurveEditorModule& CurveEditorModule = FModuleManager::LoadModuleChecked<ICurveEditorModule>(CurveEditorName);
+				CurveEditorModule.UnregisterView(FWaveTableCurveModel::WaveTableViewId);
+				FWaveTableCurveModel::WaveTableViewId = ECurveEditorViewID::Invalid;
+			}
 		}
 	} // namespace Editor
 } // namespace WaveTable
diff --git a/Engine/Plugins/Runtime/WaveTable/Source/WaveTableEditor/Private/WaveTableEditorModule.h b/Engine/Plugins/Runtime/WaveTable/Source/WaveTableEditor/Private/WaveTableEditorModule.h
--- a/Engine/Plugins/Runtime/WaveTable/Source/WaveTableEditor/Private/WaveTableEditorModule.h
+++ b/Engine/Plugins/Runtime/WaveTable/Source/WaveTableEditor/Private/WaveTableEditorModule.h
@@ -1,6 +1,8 @@
 // Copyright Epic Games, Inc. All Rights Reserved.
 #pragma once
 
+#include "AssetTypeActions_Base.h"
+#include "Containers/Array.h"
 #include "Logging/LogMacros.h"
 #include "Modules/ModuleManager.h"
 
@@ -19,6 +21,9 @@ namespace WaveTable
 			virtual void StartupModule() override;
 
 			virtual void ShutdownModule() override;
+
+		private:
+			TArray<TSharedPtr<FAssetTypeActions_Base>> AssetActions;
 		};
 	} // namespace Editor
 } // namespace WaveTable
diff --git a/Engine/Plugins/Runtime/WaveTable/Source/WaveTableEditor/Private/WaveTableTransformLayout.cpp b/Engine/Plugins/Runtime/WaveTable/Source/WaveTableEditor/Private/WaveTableTransformLayout.cpp
new file mode 100644
--- /dev/null
+++ b/Engine/Plugins/Runtime/WaveTable/Source/WaveTableEditor/Private/WaveTableTransformLayout.cpp
@@ -0,0 +1,237 @@
+// Copyright Epic Games, Inc. All Rights Reserved.
+#include "WaveTableTransformLayout.h"
+
+#include "DetailCategoryBuilder.h"
+#include "DetailLayoutBuilder.h"
+#include "DetailWidgetRow.h"
+#include "IDetailChildrenBuilder.h"
+#include "IDetailPropertyRow.h"
+#include "InputCoreTypes.h"
+#include "PropertyRestriction.h"
+#include "SCurveEditor.h"
+#include "ScopedTransaction.h"
+#include "Styling/AppStyle.h"
+#include "WaveTableBank.h"
+#include "WaveTableFileUtilities.h"
+#include "Widgets/Input/SButton.h"
+#include "Widgets/Input/SEditableTextBox.h"
+#include "Widgets/Input/SNumericEntryBox.h"
+#include "Widgets/Layout/SSeparator.h"
+#include "Widgets/Layout/SWidgetSwitcher.h"
+#include "Widgets/Notifications/SNotificationList.h"
+#include "Widgets/Text/SRichTextBlock.h"
+#include "Widgets/Text/STextBlock.h"
+#include "Widgets/Views/SListView.h"
+
+#define LOCTEXT_NAMESPACE "WaveTable"
+
+
+namespace WaveTable
+{
+	namespace Editor
+	{
+		void FTransformLayoutCustomizationBase::CustomizeHeader(TSharedRef<IPropertyHandle> StructPropertyHandle, FDetailWidgetRow& HeaderRow, IPropertyTypeCustomizationUtils& CustomizationUtils)
+		{
+		}
+
+		void FTransformLayoutCustomizationBase::CustomizeChildren(TSharedRef<IPropertyHandle> StructPropertyHandle, IDetailChildrenBuilder& ChildBuilder, IPropertyTypeCustomizationUtils& StructCustomizationUtils)
+		{
+			uint32 NumChildren;
+			StructPropertyHandle->GetNumChildren(NumChildren);
+
+			FProperty* Property = nullptr;
+			StructPropertyHandle->GetValue(Property);
+
+			TMap<FName, TSharedPtr<IPropertyHandle>> PropertyHandles;
+
+			for (uint32 ChildIndex = 0; ChildIndex < NumChildren; ++ChildIndex)
+			{
+				TSharedRef<IPropertyHandle> ChildHandle = StructPropertyHandle->GetChildHandle(ChildIndex).ToSharedRef();
+				const FName PropertyName = ChildHandle->GetProperty()->GetFName();
+				PropertyHandles.Add(PropertyName, ChildHandle);
+			}
+
+			CurveHandle = PropertyHandles.FindChecked(GET_MEMBER_NAME_CHECKED(FWaveTableTransform, Curve)).ToSharedRef();
+			ChildBuilder.AddProperty(CurveHandle->AsShared());
+
+			TSharedRef<IPropertyHandle> ScalarHandle = PropertyHandles.FindChecked(GET_MEMBER_NAME_CHECKED(FWaveTableTransform, Scalar)).ToSharedRef();
+
+			ChildBuilder.AddProperty(ScalarHandle)
+				.EditCondition(TAttribute<bool>::Create([this]() { return IsScaleableCurve(); }), nullptr)
+				.Visibility(TAttribute<EVisibility>::Create([this]() { return IsScaleableCurve() ? EVisibility::Visible : EVisibility::Hidden; }));
+
+			TSharedRef<IPropertyHandle> SharedCurveHandle = PropertyHandles.FindChecked(GET_MEMBER_NAME_CHECKED(FWaveTableTransform, CurveShared)).ToSharedRef();
+			ChildBuilder.AddProperty(SharedCurveHandle)
+				.EditCondition(TAttribute<bool>::Create([this]() { return GetCurve() == EWaveTableCurve::Shared; }), nullptr)
+				.Visibility(TAttribute<EVisibility>::Create([this]() { return GetCurve() == EWaveTableCurve::Shared ? EVisibility::Visible : EVisibility::Hidden; }));
+
+			TAttribute<EVisibility> IsWaveTableVisibilityAttribute = TAttribute<EVisibility>::Create([this]()
+			{
+				return GetCurve() == EWaveTableCurve::File ? EVisibility::Visible : EVisibility::Hidden;
+			});
+
+			WaveTableOptionsHandle = PropertyHandles.FindChecked(GET_MEMBER_NAME_CHECKED(FWaveTableTransform, WaveTableSettings)).ToSharedRef();
+
+			uint32 NumWaveTableOptions = 0;
+			if (WaveTableOptionsHandle->GetNumChildren(NumWaveTableOptions))
+			{
+				for (int32 i = 0; i < static_cast<int32>(NumWaveTableOptions); ++i)
+				{
+					TSharedPtr<IPropertyHandle> ChildProperty = WaveTableOptionsHandle->GetChildHandle(i);
+					if (ChildProperty->GetProperty()->GetFName() != GET_MEMBER_NAME_CHECKED(FWaveTableSettings, FilePath))
+					{
+						ChildProperty->SetOnPropertyValueChanged(FSimpleDelegate::CreateLambda([this]
+						{
+							RefreshWaveTable();
+						}));
+					}
+				}
+
+				ChannelIndexHandle = WaveTableOptionsHandle->GetChildHandle(GET_MEMBER_NAME_CHECKED(FWaveTableSettings, ChannelIndex));
+
+				FilePathHandle = WaveTableOptionsHandle->GetChildHandle(GET_MEMBER_NAME_CHECKED(FWaveTableSettings, FilePath));
+				FilePathHandle->SetOnChildPropertyValueChanged(FSimpleDelegate::CreateLambda([this] { CachePCMFromFile(); }));
+				FilePathHandle = FilePathHandle->GetChildHandle(GET_MEMBER_NAME_CHECKED(FFilePath, FilePath));
+			}
+
+			ChildBuilder.AddProperty(WaveTableOptionsHandle->AsShared())
+				.EditCondition(TAttribute<bool>::Create([this]() { return GetCurve() == EWaveTableCurve::File; }), nullptr)
+				.Visibility(IsWaveTableVisibilityAttribute);
+		}
+
+		void FTransformLayoutCustomizationBase::CachePCMFromFile()
+		{
+			if (FWaveTableTransform* Transform = GetTransform())
+			{
+				bool bLoadSucceeded = false;
+				FString FilePath;
+				if (ensure(FilePathHandle->GetValue(FilePath) == FPropertyAccess::Success))
+				{
+					int32 ChannelIndex = 0;
+					if (ensure(ChannelIndexHandle->GetValue(ChannelIndex) == FPropertyAccess::Success))
+					{
+						TArray<float>& SourcePCMData = Transform->WaveTableSettings.SourcePCMData;
+						WaveTable::Editor::FileUtilities::LoadPCMChannel(FilePath, ChannelIndex, SourcePCMData);
+						bLoadSucceeded = true;
+					}
+				}
+
+				if (!bLoadSucceeded)
+				{
+					Transform->WaveTableSettings.SourcePCMData.Empty();
+				}
+			}
+
+			RefreshWaveTable();
+		}
+
+		void FTransformLayoutCustomizationBase::RefreshWaveTable() const
+		{
+			if (FWaveTableTransform* Transform = GetTransform())
+			{
+				if (const EWaveTableResolution* Resolution = GetResolution())
+				{
+					const bool bBipolar = IsBipolar();
+					Transform->CacheWaveTable(*Resolution, bBipolar);
+				}
+			}
+		}
+
+		bool FTransformLayoutCustomizationBase::IsScaleableCurve() const
+		{
+			static const TArray<EWaveTableCurve> ScalarFilters = { EWaveTableCurve::Exp, EWaveTableCurve::Exp_Inverse, EWaveTableCurve::Log };
+			return ScalarFilters.Contains(GetCurve());
+		}
+
+		EWaveTableCurve FTransformLayoutCustomizationBase::GetCurve() const
+		{
+			if (!CurveHandle.IsValid())
+			{
+				return EWaveTableCurve::Linear;
+			}
+
+			uint8 CurveValue = static_cast<uint8>(EWaveTableCurve::Linear);
+			CurveHandle->GetValue(CurveValue);
+			return static_cast<EWaveTableCurve>(CurveValue);
+		}
+
+		int32 FTransformLayoutCustomizationBase::GetOwningArrayIndex() const
+		{
+			int32 TableIndex = INDEX_NONE;
+			if (ensure(WaveTableOptionsHandle.IsValid()))
+			{
+				TSharedPtr<IPropertyHandleArray> ParentArray;
+				TSharedPtr<IPropertyHandle> ChildHandle = WaveTableOptionsHandle;
+				while (ChildHandle.IsValid() && TableIndex == INDEX_NONE)
+				{
+					TableIndex = ChildHandle->GetIndexInArray();
+					ChildHandle = ChildHandle->GetParentHandle();
+				}
+			}
+
+			return TableIndex;
+		}
+
+		const EWaveTableResolution* FTransformLayoutCustomization::GetResolution() const
+		{
+			if (ensure(WaveTableOptionsHandle.IsValid()))
+			{
+				TArray<UObject*> OuterObjects;
+				WaveTableOptionsHandle->GetOuterObjects(OuterObjects);
+				if (OuterObjects.Num() == 1)
+				{
+					if (UWaveTableBank* Bank = Cast<UWaveTableBank>(OuterObjects.Last()))
+					{
+						return &Bank->Resolution;
+					}
+				}
+			}
+
+			return nullptr;
+		}
+
+		bool FTransformLayoutCustomization::IsBipolar() const
+		{
+			if (ensure(WaveTableOptionsHandle.IsValid()))
+			{
+				TArray<UObject*> OuterObjects;
+				WaveTableOptionsHandle->GetOuterObjects(OuterObjects);
+				if (OuterObjects.Num() == 1)
+				{
+					if (UWaveTableBank* Bank = Cast<UWaveTableBank>(OuterObjects.Last()))
+					{
+						return Bank->bBipolar;
+					}
+				}
+			}
+
+			return false;
+		}
+
+		FWaveTableTransform* FTransformLayoutCustomization::GetTransform() const
+		{
+			if (ensure(WaveTableOptionsHandle.IsValid()))
+			{
+				TArray<UObject*> OuterObjects;
+				WaveTableOptionsHandle->GetOuterObjects(OuterObjects);
+				if (OuterObjects.Num() == 1)
+				{
+					if (UWaveTableBank* Bank = Cast<UWaveTableBank>(OuterObjects.Last()))
+					{
+						const int32 TableIndex = GetOwningArrayIndex();
+						if (TableIndex != INDEX_NONE)
+						{
+							if (TableIndex < Bank->Entries.Num())
+							{
+								return &Bank->Entries[TableIndex].Transform;
+							}
+						}
+					}
+				}
+			}
+
+			return nullptr;
+		}
+	} // namespace Editor
+} // namespace WaveTable
+#undef LOCTEXT_NAMESPACE // AudioModulation
diff --git a/Engine/Plugins/Runtime/WaveTable/Source/WaveTableEditor/Public/WaveTableBankEditor.h b/Engine/Plugins/Runtime/WaveTable/Source/WaveTableEditor/Public/WaveTableBankEditor.h
--- a/Engine/Plugins/Runtime/WaveTable/Source/WaveTableEditor/Public/WaveTableBankEditor.h
+++ b/Engine/Plugins/Runtime/WaveTable/Source/WaveTableEditor/Public/WaveTableBankEditor.h
@@ -17,6 +17,7 @@
 
 // Forward Declarations
 class FCurveEditor;
+struct FWaveTableTransform;
 class IToolkitHost;
 class SCurveEditorPanel;
 class UCurveFloat;
@@ -25,11 +26,11 @@ namespace WaveTable
 {
 	namespace Editor
 	{
-		class WAVETABLEEDITOR_API FWaveTableBankEditor : public FAssetEditorToolkit, public FNotifyHook, public FEditorUndoClient
+		class WAVETABLEEDITOR_API FBankEditorBase : public FAssetEditorToolkit, public FNotifyHook, public FEditorUndoClient
 		{
 		public:
-			FWaveTableBankEditor();
-			virtual ~FWaveTableBankEditor() = default;
+			FBankEditorBase();
+			virtual ~FBankEditorBase() = default;
 
 			void Init(const EToolkitMode::Type Mode, const TSharedPtr<IToolkitHost>& InitToolkitHost, UObject* InParentObject);
 
@@ -59,24 +60,27 @@ namespace WaveTable
 			virtual void PostUndo(bool bSuccess) override;
 			virtual void PostRedo(bool bSuccess) override;
 
-			virtual bool GetIsPropertyEditorDisabled() const = 0;
+			virtual bool GetIsPropertyEditorDisabled() const { return false; }
 
-			/** Generates expression curve at the given index. */
-			virtual void GenerateExpressionCurve(FCurveData& OutCurveData, int32 InInputIndex, EWaveTableCurveSource InSource, bool bInIsUnset = false) = 0;
+			virtual EWaveTableResolution GetBankResolution() const = 0;
+			virtual bool GetBankIsBipolar() const = 0;
+
+			virtual TUniquePtr<FWaveTableCurveModel> ConstructCurveModel(FRichCurve& InRichCurve, UObject* InParentObject, EWaveTableCurveSource InSource) = 0;
 
-			virtual TUniquePtr<FWaveTableCurveModelBase> ConstructCurveModel(FRichCurve& InRichCurve, UObject* InParentObject, EWaveTableCurveSource InSource) = 0;
+			virtual FWaveTableTransform* GetTransform (int32 InIndex) const = 0;
 
-			virtual EWaveTableCurve GetCurveType(int32 InInputIndex) const = 0;
-			virtual FRichCurve& GetCustomCurveChecked(int32 InInputIndex) const = 0;
-			virtual UCurveFloat* GetSharedCurve(int32 InInputIndex) const = 0;
-			virtual int32 GetNumCurves() const = 0;
+			virtual int32 GetNumTransforms() const = 0;
 
-			void SetCurve(int32 InInputIndex, FRichCurve& InRichCurve, EWaveTableCurveSource InSource);
+			void SetCurve(int32 InTransformIndex, FRichCurve& InRichCurve, EWaveTableCurveSource InSource);
 
 		private:
-			void ResetCurves();
+			/** Generates expression curve at the given index. */
+			void GenerateExpressionCurve(FCurveData& OutCurveData, int32 InTransformIndex, EWaveTableCurveSource InSource, bool bInIsUnset = false);
+
 			void InitCurves();
 
+			void ResetCurves();
+
 			/** Updates & redraws curves. */
 			void RefreshCurves();
 
@@ -93,9 +97,9 @@ namespace WaveTable
 			static void TrimKeys(FRichCurve& OutCurve);
 
 			/** Clears the expression curve at the given input index */
-			void ClearExpressionCurve(int32 InInputIndex);
+			void ClearExpressionCurve(int32 InTransformIndex);
 
-			bool RequiresNewCurve(int32 InInputIndex, const FRichCurve& InRichCurve) const;
+			bool RequiresNewCurve(int32 InTransformIndex, const FRichCurve& InRichCurve) const;
 
 			TSharedPtr<FUICommandList> ToolbarCurveTargetCommands;
 
@@ -112,5 +116,20 @@ namespace WaveTable
 			static const FName CurveTabId;
 			static const FName PropertiesTabId;
 		};
+
+		class FBankEditor : public WaveTable::Editor::FBankEditorBase
+		{
+		public:
+			FBankEditor() = default;
+			virtual ~FBankEditor() = default;
+
+		protected:
+			virtual EWaveTableResolution GetBankResolution() const override;
+			virtual bool GetBankIsBipolar() const override;
+			virtual int32 GetNumTransforms() const override;
+			virtual FWaveTableTransform* GetTransform(int32 InIndex) const override;
+
+			virtual TUniquePtr<WaveTable::Editor::FWaveTableCurveModel> ConstructCurveModel(FRichCurve& InRichCurve, UObject* InParentObject, EWaveTableCurveSource InSource) override;
+		};
 	} // namespace Editor
 } // namespace WaveTable
\ No newline at end of file
diff --git a/Engine/Plugins/Runtime/WaveTable/Source/WaveTableEditor/Public/WaveTableCurveEditorViewStacked.h b/Engine/Plugins/Runtime/WaveTable/Source/WaveTableEditor/Public/WaveTableCurveEditorViewStacked.h
--- a/Engine/Plugins/Runtime/WaveTable/Source/WaveTableEditor/Public/WaveTableCurveEditorViewStacked.h
+++ b/Engine/Plugins/Runtime/WaveTable/Source/WaveTableEditor/Public/WaveTableCurveEditorViewStacked.h
@@ -15,6 +15,7 @@
 
 // Forward Declarations
 class UCurveFloat;
+struct FWaveTableTransform;
 
 
 UENUM()
@@ -31,34 +32,30 @@ namespace WaveTable
 {
 	namespace Editor
 	{
-		class WAVETABLEEDITOR_API FWaveTableCurveModelBase : public FRichCurveEditorModelRaw
+		class WAVETABLEEDITOR_API FWaveTableCurveModel : public FRichCurveEditorModelRaw
 		{
 		public:
-			static ECurveEditorViewID ViewId;
+			static ECurveEditorViewID WaveTableViewId;
 
-			FWaveTableCurveModelBase(FRichCurve& InRichCurve, UObject* InOwner, EWaveTableCurveSource InSource);
+			FWaveTableCurveModel(FRichCurve& InRichCurve, UObject* InOwner, EWaveTableCurveSource InSource);
 
 			const FText& GetAxesDescriptor() const;
 			const UObject* GetParentObject() const;
 			EWaveTableCurveSource GetSource() const;
-			void Refresh(int32 InCurveIndex);
+			void Refresh(const FWaveTableTransform& InTransform, int32 InCurveIndex);
 
+			virtual ECurveEditorViewID GetViewId() const { return WaveTableViewId; }
 			virtual bool IsReadOnly() const override;
 			virtual FLinearColor GetColor() const override;
 
-			int32 GetCurveIndex() const
-			{
-				return CurveIndex;
-			}
+			int32 GetCurveIndex() const { return CurveIndex; }
 
 		protected:
-			virtual void RefreshCurveDescriptorText(FText& OutShortDisplayName, FText& OutInputAxisName, FText& OutOutputAxisName) = 0;
-
-			virtual UCurveFloat* GetSharedCurve() = 0;
-
-			virtual FColor GetCurveColor() const = 0;
-			virtual bool GetPropertyEditorDisabled() const = 0;
-			virtual FText GetPropertyEditorDisabledText() const = 0;
+			virtual void RefreshCurveDescriptorText(const FWaveTableTransform& InTransform, FText& OutShortDisplayName, FText& OutInputAxisName, FText& OutOutputAxisName);
+			virtual FWaveTableTransform* GetTransform();
+			virtual FColor GetCurveColor() const;
+			virtual bool GetPropertyEditorDisabled() const;
+			virtual FText GetPropertyEditorDisabledText() const;
 
 			TWeakObjectPtr<UObject> ParentObject;
 
@@ -88,8 +85,8 @@ namespace WaveTable
 			virtual void DrawViewGrids(const FGeometry& AllottedGeometry, const FSlateRect& MyCullingRect, FSlateWindowElementList& OutDrawElements, int32 BaseLayerId, ESlateDrawEffect DrawEffects) const override;
 			virtual void DrawLabels(const FGeometry& AllottedGeometry, const FSlateRect& MyCullingRect, FSlateWindowElementList& OutDrawElements, int32 BaseLayerId, ESlateDrawEffect DrawEffects) const override;
 
-			virtual void FormatInputLabel(const FWaveTableCurveModelBase& EditorModel, const FNumberFormattingOptions& InLabelFormat, FText& InOutLabel) const { }
-			virtual void FormatOutputLabel(const FWaveTableCurveModelBase& EditorModel, const FNumberFormattingOptions& InLabelFormat, FText& InOutLabel) const { }
+			virtual void FormatInputLabel(const FWaveTableCurveModel& EditorModel, const FNumberFormattingOptions& InLabelFormat, FText& InOutLabel) const { }
+			virtual void FormatOutputLabel(const FWaveTableCurveModel& EditorModel, const FNumberFormattingOptions& InLabelFormat, FText& InOutLabel) const { }
 
 		private:
 			struct FGridDrawInfo
diff --git a/Engine/Plugins/Runtime/WaveTable/Source/WaveTableEditor/Public/WaveTableTransformLayout.h b/Engine/Plugins/Runtime/WaveTable/Source/WaveTableEditor/Public/WaveTableTransformLayout.h
new file mode 100644
--- /dev/null
+++ b/Engine/Plugins/Runtime/WaveTable/Source/WaveTableEditor/Public/WaveTableTransformLayout.h
@@ -0,0 +1,58 @@
+// Copyright Epic Games, Inc. All Rights Reserved.
+#pragma once
+
+#include "Curves/SimpleCurve.h"
+#include "IPropertyTypeCustomization.h"
+#include "Layout/Visibility.h"
+#include "Misc/Attribute.h"
+#include "PropertyHandle.h"
+#include "SCurveEditor.h"
+#include "WaveTableSettings.h"
+#include "WaveTableTransform.h"
+
+
+namespace WaveTable
+{
+	namespace Editor
+	{
+		class WAVETABLEEDITOR_API FTransformLayoutCustomizationBase : public IPropertyTypeCustomization
+		{
+		public:
+			//~ Begin IPropertyTypeCustomization
+			virtual void CustomizeHeader(TSharedRef<IPropertyHandle> StructPropertyHandle, FDetailWidgetRow& HeaderRow, IPropertyTypeCustomizationUtils& StructCustomizationUtils) override;
+			virtual void CustomizeChildren(TSharedRef<IPropertyHandle> StructPropertyHandle, IDetailChildrenBuilder& ChildBuilder, IPropertyTypeCustomizationUtils& StructCustomizationUtils) override;
+			//~ End IPropertyTypeCustomization
+
+		protected:
+			EWaveTableCurve GetCurve() const;
+
+			virtual const EWaveTableResolution* GetResolution() const = 0;
+			virtual FWaveTableTransform* GetTransform() const = 0;
+			virtual bool IsBipolar() const = 0;
+
+			void CachePCMFromFile();
+			int32 GetOwningArrayIndex() const;
+
+			bool IsScaleableCurve() const;
+			void RefreshWaveTable() const;
+
+			TSharedPtr<IPropertyHandle> CurveHandle;
+			TSharedPtr<IPropertyHandle> ChannelIndexHandle;
+			TSharedPtr<IPropertyHandle> FilePathHandle;
+			TSharedPtr<IPropertyHandle> WaveTableOptionsHandle;
+		};
+
+		class WAVETABLEEDITOR_API FTransformLayoutCustomization : public FTransformLayoutCustomizationBase
+		{
+		public:
+			static TSharedRef<IPropertyTypeCustomization> MakeInstance()
+			{
+				return MakeShared<FTransformLayoutCustomization>();
+			}
+
+			virtual const EWaveTableResolution* GetResolution() const override;
+			virtual bool IsBipolar() const override;
+			virtual FWaveTableTransform* GetTransform() const override;
+		};
+	} // namespace Editor
+} // namespace WaveTable
