Fix for Linux Editor Compile Error
#rnx

diff --git a/Engine/Plugins/Runtime/AudioModulation/Source/AudioModulation/Public/SoundModulationPatch.h b/Engine/Plugins/Runtime/AudioModulation/Source/AudioModulation/Public/SoundModulationPatch.h
--- a/Engine/Plugins/Runtime/AudioModulation/Source/AudioModulation/Public/SoundModulationPatch.h
+++ b/Engine/Plugins/Runtime/AudioModulation/Source/AudioModulation/Public/SoundModulationPatch.h
@@ -15,7 +15,7 @@ class USoundControlBus;
 
 
 USTRUCT(BlueprintType)
-struct FSoundModulationTransform : public FWaveTableTransform
+struct AUDIOMODULATION_API FSoundModulationTransform : public FWaveTableTransform
 {
 	GENERATED_USTRUCT_BODY()
 };
@@ -61,7 +61,7 @@ struct AUDIOMODULATION_API FSoundControlModulationPatch
 	TArray<FSoundControlModulationInput> Inputs;
 
 	UPROPERTY(EditAnywhere, BlueprintReadWrite, Category = Inputs)
-	EWaveTableResolution WaveTableResolution;
+	EWaveTableResolution WaveTableResolution = EWaveTableResolution::None;
 };
 
 UCLASS(config = Engine, editinlinenew, BlueprintType)
