Potential fix for interfaces reporting not being registered on MetaSound Asset load
#rnx

diff --git a/Engine/Plugins/Runtime/AudioModulation/Source/AudioModulation/Private/AudioModulation.cpp b/Engine/Plugins/Runtime/AudioModulation/Source/AudioModulation/Private/AudioModulation.cpp
--- a/Engine/Plugins/Runtime/AudioModulation/Source/AudioModulation/Private/AudioModulation.cpp
+++ b/Engine/Plugins/Runtime/AudioModulation/Source/AudioModulation/Private/AudioModulation.cpp
@@ -267,7 +267,11 @@ void FAudioModulationModule::StartupModule()
 
 #if WITH_AUDIOMODULATION_METASOUND_SUPPORT
 	UE_LOG(LogAudioModulation, Log, TEXT("Registering Modulation MetaSound Nodes..."));
-	// flush node registration queue to guarantee AudioModulation DataTypes/Nodes are ready prior to assets loading
+
+	// All MetaSound interfaces are required to be loaded prior to registering & loading MetaSound assets,
+	// so check that the MetaSoundEngine is loaded prior to pending Modulation defined classes
+	FModuleManager::Get().LoadModuleChecked("MetasoundEngine");
+
 	FMetasoundFrontendRegistryContainer::Get()->RegisterPendingNodes();
 #endif // WITH_AUDIOMODULATION_METASOUND_SUPPORT
 
