Fix for Modulation Parameter�s Unit not displaying

diff --git a/Engine/Plugins/Runtime/AudioModulation/Source/AudioModulationEditor/Private/Layouts/SoundControlBusMixStageLayout.cpp b/Engine/Plugins/Runtime/AudioModulation/Source/AudioModulationEditor/Private/Layouts/SoundControlBusMixStageLayout.cpp
--- a/Engine/Plugins/Runtime/AudioModulation/Source/AudioModulationEditor/Private/Layouts/SoundControlBusMixStageLayout.cpp
+++ b/Engine/Plugins/Runtime/AudioModulation/Source/AudioModulationEditor/Private/Layouts/SoundControlBusMixStageLayout.cpp
@@ -63,6 +63,12 @@ namespace AudioModulationEditorUtils
 					UnitValueHandle->SetValue(UnitValue, EPropertyValueSetFlags::NotTransactable);
 				}
 			}
+			else
+			{
+				float NormalizedValue = 1.0f;
+				NormalizedValueHandle->GetValue(NormalizedValue);
+				UnitValueHandle->SetValue(NormalizedValue);
+			}
 		}
 	}
 
@@ -79,6 +85,12 @@ namespace AudioModulationEditorUtils
 					NormalizedValueHandle->SetValue(NormalizedValue);
 				}
 			}
+			else
+			{
+				float UnitValue = 1.0f;
+				UnitValueHandle->GetValue(UnitValue);
+				NormalizedValueHandle->SetValue(UnitValue);
+			}
 		}
 	}
 }
@@ -184,7 +196,7 @@ void FSoundControlBusMixStageLayoutCustomization::CustomizeChildren(TSharedRef<I
 					.Text(TAttribute<FText>::Create([BusHandle]()
 					{
 						USoundModulationParameter* Parameter = AudioModulationEditorUtils::GetParameterFromBus(BusHandle);
-						if (Parameter && Parameter->RequiresUnitConversion())
+						if (Parameter)
 						{
 							return Parameter->Settings.UnitDisplayName;
 						}
@@ -213,7 +225,7 @@ void FSoundControlBusMixStageLayoutCustomization::CustomizeChildren(TSharedRef<I
 	{
 		if (USoundModulationParameter* Parameter = AudioModulationEditorUtils::GetParameterFromBus(BusHandle))
 		{
-			return Parameter->RequiresUnitConversion() ? EVisibility::Visible : EVisibility::Hidden;
+			return !Parameter->Settings.UnitDisplayName.IsEmpty() ? EVisibility::Visible : EVisibility::Hidden;
 		}
 
 		return EVisibility::Hidden;
@@ -224,7 +236,7 @@ void FSoundControlBusMixStageLayoutCustomization::CustomizeChildren(TSharedRef<I
 	{
 		if (USoundModulationParameter* Parameter = AudioModulationEditorUtils::GetParameterFromBus(BusHandle))
 		{
-			return Parameter->RequiresUnitConversion() ? EVisibility::Hidden : EVisibility::Visible;
+			return !Parameter->Settings.UnitDisplayName.IsEmpty() ? EVisibility::Hidden : EVisibility::Visible;
 		}
 
 		return EVisibility::Visible;
