- Set the MixStage to the referenced bus' ModulationParameter default normalized/unit values when changing input bus
- Add 'Reset To Default' for MixStages

diff --git a/Engine/Plugins/Runtime/AudioModulation/Source/AudioModulation/Private/SoundControlBusMixProxy.cpp b/Engine/Plugins/Runtime/AudioModulation/Source/AudioModulation/Private/SoundControlBusMixProxy.cpp
--- a/Engine/Plugins/Runtime/AudioModulation/Source/AudioModulation/Private/SoundControlBusMixProxy.cpp
+++ b/Engine/Plugins/Runtime/AudioModulation/Source/AudioModulation/Private/SoundControlBusMixProxy.cpp
@@ -60,7 +60,7 @@ namespace AudioModulation
 			}
 			else
 			{
-				UE_LOG(LogAudioModulation, Warning,
+				UE_LOG(LogAudioModulation, VeryVerbose,
 					TEXT("USoundControlBusMix '%s' has stage with no bus specified. "
 						"Mix instance initialized with stage ignored."),
 					*InBusMix.GetFullName());
diff --git a/Engine/Plugins/Runtime/AudioModulation/Source/AudioModulationEditor/Private/Layouts/SoundControlBusMixStageLayout.cpp b/Engine/Plugins/Runtime/AudioModulation/Source/AudioModulationEditor/Private/Layouts/SoundControlBusMixStageLayout.cpp
--- a/Engine/Plugins/Runtime/AudioModulation/Source/AudioModulationEditor/Private/Layouts/SoundControlBusMixStageLayout.cpp
+++ b/Engine/Plugins/Runtime/AudioModulation/Source/AudioModulationEditor/Private/Layouts/SoundControlBusMixStageLayout.cpp
@@ -93,6 +93,37 @@ namespace AudioModulationEditorUtils
 			}
 		}
 	}
+
+	void HandleSetDefault(TSharedRef<IPropertyHandle> BusHandle, TSharedRef<IPropertyHandle> UnitValueHandle, TSharedRef<IPropertyHandle> NormalizedValueHandle)
+	{
+		if (USoundModulationParameter* Parameter = AudioModulationEditorUtils::GetParameterFromBus(BusHandle))
+		{
+			NormalizedValueHandle->SetValue(Parameter->Settings.ValueNormalized);
+			UnitValueHandle->SetValue(Parameter->GetUnitDefault());
+		}
+		else
+		{
+			const Audio::FModulationParameter& DefaultParam = Audio::GetModulationParameter({ });
+
+			float NormalizedValue = 1.0f;
+			DefaultParam.GetDefaultNormalizedConversionFunction()(NormalizedValue);
+			NormalizedValueHandle->SetValue(NormalizedValue);
+			UnitValueHandle->SetValue(DefaultParam.DefaultValue);
+		}
+	}
+
+	bool HandleIsDefaultValue(TSharedRef<IPropertyHandle> BusHandle, TSharedRef<IPropertyHandle> NormalizedValueHandle)
+	{
+		float ValueNormalized = 1.0f;
+		NormalizedValueHandle->GetValue(ValueNormalized);
+
+		if (USoundModulationParameter* Parameter = AudioModulationEditorUtils::GetParameterFromBus(BusHandle))
+		{
+			return FMath::IsNearlyEqual(Parameter->Settings.ValueNormalized, ValueNormalized);
+		}
+
+		return FMath::IsNearlyEqual(ValueNormalized, 1.0f);
+	}
 }
 
 void FSoundControlBusMixStageLayoutCustomization::CustomizeHeader(TSharedRef<IPropertyHandle> StructPropertyHandle, FDetailWidgetRow& HeaderRow, IPropertyTypeCustomizationUtils& CustomizationUtils)
@@ -140,19 +171,18 @@ void FSoundControlBusMixStageLayoutCustomization::CustomizeChildren(TSharedRef<I
 	}
 
 	TSharedRef<IPropertyHandle> BusHandle = PropertyHandles.FindChecked(GET_MEMBER_NAME_CHECKED(FSoundControlBusMixStage, Bus)).ToSharedRef();
+	TSharedRef<IPropertyHandle> NormalizedValueHandle = StructPropertyHandle->GetChildHandle(GET_MEMBER_NAME_CHECKED(FSoundModulationMixValue, TargetValue)).ToSharedRef();
+	TSharedRef<IPropertyHandle> UnitValueHandle = StructPropertyHandle->GetChildHandle(GET_MEMBER_NAME_CHECKED(FSoundModulationMixValue, TargetUnitValue)).ToSharedRef();
 
-	TAttribute<EVisibility> BusInfoVisibility = TAttribute<EVisibility>::Create([BusHandle]()
+	auto ResetLambda = [BusHandle, NormalizedValueHandle, UnitValueHandle]()
 	{
-		UObject* Bus = nullptr;
-		BusHandle->GetValue(Bus);
-		return Bus && Bus->IsA<USoundControlBus>() ? EVisibility::Visible : EVisibility::Hidden;
-	});
+		AudioModulationEditorUtils::HandleSetDefault(BusHandle, UnitValueHandle, NormalizedValueHandle);
+	};
 
+	BusHandle->SetOnPropertyValueChanged(FSimpleDelegate::CreateLambda(ResetLambda));
+	BusHandle->SetOnPropertyResetToDefault(FSimpleDelegate::CreateLambda(ResetLambda));
 	ChildBuilder.AddProperty(BusHandle);
 
-	TSharedRef<IPropertyHandle> NormalizedValueHandle = StructPropertyHandle->GetChildHandle(GET_MEMBER_NAME_CHECKED(FSoundModulationMixValue, TargetValue)).ToSharedRef();
-	TSharedRef<IPropertyHandle> UnitValueHandle = StructPropertyHandle->GetChildHandle(GET_MEMBER_NAME_CHECKED(FSoundModulationMixValue, TargetUnitValue)).ToSharedRef();
-
 	// When editor opens, set unit value in case bus unit has changed while editor was closed.
 	AudioModulationEditorUtils::HandleConvertNormalizedToUnit(BusHandle, NormalizedValueHandle, UnitValueHandle);
 
@@ -167,6 +197,14 @@ void FSoundControlBusMixStageLayoutCustomization::CustomizeChildren(TSharedRef<I
 		AudioModulationEditorUtils::HandleConvertNormalizedToUnit(BusHandle, NormalizedValueHandle, UnitValueHandle);
 	}));
 
+	FResetToDefaultOverride ResetToDefault = FResetToDefaultOverride::Create(
+		TAttribute<bool>::Create([BusHandle, NormalizedValueHandle]()
+		{
+			return !AudioModulationEditorUtils::HandleIsDefaultValue(BusHandle, NormalizedValueHandle);
+		}),
+		FSimpleDelegate::CreateLambda(ResetLambda)
+	);
+
 	ChildBuilder.AddCustomRow(StructPropertyHandle->GetPropertyDisplayName())
 	.NameContent()
 	[
@@ -229,7 +267,8 @@ void FSoundControlBusMixStageLayoutCustomization::CustomizeChildren(TSharedRef<I
 		}
 
 		return EVisibility::Hidden;
-	}));
+	}))
+	.OverrideResetToDefault(ResetToDefault);
 
 	ChildBuilder.AddProperty(NormalizedValueHandle)
 	.Visibility(TAttribute<EVisibility>::Create([BusHandle]()
@@ -240,7 +279,8 @@ void FSoundControlBusMixStageLayoutCustomization::CustomizeChildren(TSharedRef<I
 		}
 
 		return EVisibility::Visible;
-	}));
+	}))
+	.OverrideResetToDefault(ResetToDefault);
 
 	TSharedRef<IPropertyHandle> AttackTimeHandle = StructPropertyHandle->GetChildHandle(GET_MEMBER_NAME_CHECKED(FSoundModulationMixValue, AttackTime)).ToSharedRef();
 	ChildBuilder.AddProperty(AttackTimeHandle);
diff --git a/Engine/Plugins/Runtime/AudioModulation/Source/AudioModulationEditor/Private/Layouts/SoundControlModulationPatchLayout.h b/Engine/Plugins/Runtime/AudioModulation/Source/AudioModulationEditor/Private/Layouts/SoundControlModulationPatchLayout.h
--- a/Engine/Plugins/Runtime/AudioModulation/Source/AudioModulationEditor/Private/Layouts/SoundControlModulationPatchLayout.h
+++ b/Engine/Plugins/Runtime/AudioModulation/Source/AudioModulationEditor/Private/Layouts/SoundControlModulationPatchLayout.h
@@ -13,7 +13,6 @@
 
 // Forward Declarations
 class SSearchableComboBox;
-struct FSoundModulationSettings;
 
 
 class FSoundControlModulationPatchLayoutCustomization : public IPropertyTypeCustomization
