Fix edge case in Audio Modulation Unit Display

diff --git a/Engine/Plugins/Runtime/AudioModulation/Source/AudioModulationEditor/Private/Layouts/SoundControlBusMixStageLayout.cpp b/Engine/Plugins/Runtime/AudioModulation/Source/AudioModulationEditor/Private/Layouts/SoundControlBusMixStageLayout.cpp
--- a/Engine/Plugins/Runtime/AudioModulation/Source/AudioModulationEditor/Private/Layouts/SoundControlBusMixStageLayout.cpp
+++ b/Engine/Plugins/Runtime/AudioModulation/Source/AudioModulationEditor/Private/Layouts/SoundControlBusMixStageLayout.cpp
@@ -205,6 +205,18 @@ void FSoundControlBusMixStageLayoutCustomization::CustomizeChildren(TSharedRef<I
 		FSimpleDelegate::CreateLambda(ResetLambda)
 	);
 
+	TSharedRef<SWidget> UnitValueWidget = UnitValueHandle->CreatePropertyValueWidget();
+	TSharedRef<SWidget> NormalizedValueWidget = NormalizedValueHandle->CreatePropertyValueWidget();
+	NormalizedValueWidget->SetVisibility(TAttribute<EVisibility>::Create([BusHandle]()
+	{
+		if (USoundModulationParameter* Parameter = AudioModulationEditorUtils::GetParameterFromBus(BusHandle))
+		{
+			return Parameter->RequiresUnitConversion() ? EVisibility::Visible : EVisibility::Hidden;
+		}
+	
+		return EVisibility::Hidden;
+	}));
+
 	ChildBuilder.AddCustomRow(StructPropertyHandle->GetPropertyDisplayName())
 	.NameContent()
 	[
@@ -218,11 +230,11 @@ void FSoundControlBusMixStageLayoutCustomization::CustomizeChildren(TSharedRef<I
 		[
 			SNew(SHorizontalBox)
 			+ SHorizontalBox::Slot()
-				.FillWidth(0.4f)
-				.Padding(2.0f, 0.0f, 0.0f, 0.0f)
+				.FillWidth(0.455f)
+				.Padding(0.0f, 0.0f, 0.0f, 0.0f)
 				.VAlign(VAlign_Center)
 				[
-					UnitValueHandle->CreatePropertyValueWidget()
+					UnitValueWidget
 				]
 			+ SHorizontalBox::Slot()
 				.Padding(4.0f, 0.0f, 0.0f, 0.0f)
@@ -241,45 +253,42 @@ void FSoundControlBusMixStageLayoutCustomization::CustomizeChildren(TSharedRef<I
 
 						return FText();
 					}))
-				]
-			+ SHorizontalBox::Slot()
-				.FillWidth(0.4f)
-				.Padding(4.0f, 0.0f, 0.0f, 0.0f)
-				.VAlign(VAlign_Center)
-				[
-					NormalizedValueHandle->CreatePropertyValueWidget()
-				]
-			+ SHorizontalBox::Slot()
-				.Padding(4.0f, 0.0f, 0.0f, 0.0f)
-				.VAlign(VAlign_Center)
-				.AutoWidth()
-				[
-					SNew(STextBlock)
-					.Font(IDetailLayoutBuilder::GetDetailFont())
-					.Text(LOCTEXT("SoundModulationControl_UnitValueNormalized", "Normalized"))
-				]
-	]
-	.Visibility(TAttribute<EVisibility>::Create([BusHandle]()
-	{
-		if (USoundModulationParameter* Parameter = AudioModulationEditorUtils::GetParameterFromBus(BusHandle))
-		{
-			return !Parameter->Settings.UnitDisplayName.IsEmpty() ? EVisibility::Visible : EVisibility::Hidden;
-		}
-
-		return EVisibility::Hidden;
-	}))
-	.OverrideResetToDefault(ResetToDefault);
+					.Visibility(TAttribute<EVisibility>::Create([BusHandle]()
+					{
+						if (USoundModulationParameter* Parameter = AudioModulationEditorUtils::GetParameterFromBus(BusHandle))
+						{
+							return Parameter->Settings.UnitDisplayName.IsEmptyOrWhitespace() ? EVisibility::Collapsed : EVisibility::Visible;
+						}
 
-	ChildBuilder.AddProperty(NormalizedValueHandle)
-	.Visibility(TAttribute<EVisibility>::Create([BusHandle]()
-	{
-		if (USoundModulationParameter* Parameter = AudioModulationEditorUtils::GetParameterFromBus(BusHandle))
-		{
-			return !Parameter->Settings.UnitDisplayName.IsEmpty() ? EVisibility::Hidden : EVisibility::Visible;
-		}
+						return EVisibility::Visible;
+					}))
+				]
+				+ SHorizontalBox::Slot()
+					.FillWidth(0.4f)
+					.Padding(4.0f, 0.0f, 0.0f, 0.0f)
+					.VAlign(VAlign_Center)
+					[
+						NormalizedValueWidget
+					]
+				+ SHorizontalBox::Slot()
+					.Padding(4.0f, 0.0f, 0.0f, 0.0f)
+					.VAlign(VAlign_Center)
+					.AutoWidth()
+					[
+						SNew(STextBlock)
+						.Font(IDetailLayoutBuilder::GetDetailFont())
+						.Text(LOCTEXT("SoundModulationControl_UnitValueNormalized", "Normalized"))
+						.Visibility(TAttribute<EVisibility>::Create([BusHandle]()
+						{
+							if (USoundModulationParameter* Parameter = AudioModulationEditorUtils::GetParameterFromBus(BusHandle))
+							{
+								return Parameter->RequiresUnitConversion() ? EVisibility::Visible : EVisibility::Hidden;
+							}
 
-		return EVisibility::Visible;
-	}))
+							return EVisibility::Visible;
+						}))
+					]
+		]
 	.OverrideResetToDefault(ResetToDefault);
 
 	TSharedRef<IPropertyHandle> AttackTimeHandle = StructPropertyHandle->GetChildHandle(GET_MEMBER_NAME_CHECKED(FSoundModulationMixValue, AttackTime)).ToSharedRef();
