Fix for missing normalize logic & missing pin descriptions for Mix Modulators node

diff --git a/Engine/Plugins/Runtime/AudioModulation/Source/AudioModulation/Private/MetaSoundNodes/MixModulatorsNode.cpp b/Engine/Plugins/Runtime/AudioModulation/Source/AudioModulation/Private/MetaSoundNodes/MixModulatorsNode.cpp
--- a/Engine/Plugins/Runtime/AudioModulation/Source/AudioModulation/Private/MetaSoundNodes/MixModulatorsNode.cpp
+++ b/Engine/Plugins/Runtime/AudioModulation/Source/AudioModulation/Private/MetaSoundNodes/MixModulatorsNode.cpp
@@ -25,10 +25,10 @@ namespace AudioModulation
 
 			static const FVertexInterface DefaultInterface(
 				FInputVertexInterface(
-					TInputDataVertex<FSoundModulatorAsset>("In1", FDataVertexMetadata{ LOCTEXT("MixModulatorsNode_InputModulator1Name", "In 1") }),
-					TInputDataVertex<FSoundModulatorAsset>("In2", FDataVertexMetadata{ LOCTEXT("MixModulatorsNode_InputModulator2Name", "In 2") }),
-					TInputDataVertex<FSoundModulationParameterAsset>("MixParameter", FDataVertexMetadata{ LOCTEXT("MixModulatorsNode_InputMixParameterName", "Mix Parameter") }),
-					TInputDataVertex<bool>("Normalized", FDataVertexMetadata{ LOCTEXT("MixModulatorsNode_InputNormalizedName", "Normalized") }, true)
+					TInputDataVertex<FSoundModulatorAsset>("In1", FDataVertexMetadata{ LOCTEXT("MixModulatorsNode_InputModulator1Description", "First modulator to mix together") }),
+					TInputDataVertex<FSoundModulatorAsset>("In2", FDataVertexMetadata{ LOCTEXT("MixModulatorsNode_InputModulator2Description", "Second modulator to mix together") }),
+					TInputDataVertex<FSoundModulationParameterAsset>("MixParameter", FDataVertexMetadata{ LOCTEXT("MixModulatorsNode_InputMixParameterDescription", "Parameter to use to mix parameters together") }),
+					TInputDataVertex<bool>("Normalized", FDataVertexMetadata{ LOCTEXT("MixModulatorsNode_InputNormalizedDescription", "Whether the output value should be normalized [0.0, 1.0] or converted to the unit described by the parameter.") }, true)
 				),
 				FOutputVertexInterface(
 					TOutputDataVertex<float>("Out", FDataVertexMetadata{ LOCTEXT("MixModulatorsNode_OutputModulatorValue", "Out") })
@@ -46,7 +46,7 @@ namespace AudioModulation
 				{
 					Metasound::FNodeClassName { "Modulation", "MixModulators", "" },
 					1, // Major Version
-					0, // Minor Version
+					1, // Minor Version
 					LOCTEXT("MixModulatorsNode_Name", "Mix Modulators"),
 					LOCTEXT("MixModulatorsNode_Description", "Mixes two modulators using the parameterized mix function. Returns the 'Normalized' value (0-1) if true, or the value in unit space (dB, Frequency, etc.) if set to false."),
 					AudioModulation::PluginAuthor,
@@ -150,7 +150,15 @@ namespace AudioModulation
 			const FSoundModulationParameterAsset& ParameterAsset = *Parameter;
 			if (ParameterAsset.IsValid())
 			{
-				ParameterAsset->GetParameter().MixFunction(Value1, Value2);
+				const Audio::FModulationParameter& ModParam = ParameterAsset->GetParameter();
+				ModParam.MixFunction(Value1, Value2);
+				if (!*Normalized)
+				{
+					if (ModParam.bRequiresConversion)
+					{
+						ModParam.UnitFunction(Value1);
+					}
+				}
 			}
 			else
 			{
