Leave lerp time initialized to -1.0f (not attack time) when constructing FSoundModulationMixValue from given attack time. Setting the attack time at a later point can leave a descrepancy between these two values, signaling the update logic to not update the incremental delta.
#rnx

diff --git a/Engine/Plugins/Runtime/AudioModulation/Source/AudioModulation/Private/SoundModulationValue.cpp b/Engine/Plugins/Runtime/AudioModulation/Source/AudioModulation/Private/SoundModulationValue.cpp
--- a/Engine/Plugins/Runtime/AudioModulation/Source/AudioModulation/Private/SoundModulationValue.cpp
+++ b/Engine/Plugins/Runtime/AudioModulation/Source/AudioModulation/Private/SoundModulationValue.cpp
@@ -10,7 +10,6 @@ FSoundModulationMixValue::FSoundModulationMixValue(float InValue, float InAttack
 	: TargetValue(InValue)
 	, AttackTime(InAttackTime)
 	, ReleaseTime(InReleaseTime)
-	, LerpTime(InAttackTime)
 	, Value(InValue)
 	, ActiveFade(EActiveFade::Attack)
 {
@@ -148,4 +147,3 @@ void FSoundModulationMixValue::UpdateDelta()
 
 	LastTarget = TargetValue;
 }
-
