Fix thread safety crash when forwarding ModulationDebugger data to be processed on GT for rendering to viewport
#rnx

diff --git a/Engine/Plugins/Runtime/AudioModulation/Source/AudioModulation/Private/AudioModulationDebugger.cpp b/Engine/Plugins/Runtime/AudioModulation/Source/AudioModulation/Private/AudioModulationDebugger.cpp
--- a/Engine/Plugins/Runtime/AudioModulation/Source/AudioModulation/Private/AudioModulationDebugger.cpp
+++ b/Engine/Plugins/Runtime/AudioModulation/Source/AudioModulation/Private/AudioModulationDebugger.cpp
@@ -464,25 +464,31 @@ namespace AudioModulation
 			DebugInfo.Categories = Proxy->GetDebugCategories();
 		}
 
-		FFunctionGraphTask::CreateAndDispatchWhenReady([this, RefreshedFilteredBuses, RefreshedFilteredMixes, RefreshedFilteredGenerators]()
+		FFunctionGraphTask::CreateAndDispatchWhenReady([DebuggerPtr = AsWeak(), RefreshedFilteredBuses, RefreshedFilteredMixes, RefreshedFilteredGenerators]()
 		{
-			FilteredBuses = RefreshedFilteredBuses;
-			FilteredBuses.Sort(&Debug::CompareNames<FControlBusDebugInfo>);
+			TSharedPtr<FAudioModulationDebugger> ThisDebugger = DebuggerPtr.Pin();
+			if (!ThisDebugger.IsValid())
+			{
+				return;
+			}
+
+			ThisDebugger->FilteredBuses = RefreshedFilteredBuses;
+			ThisDebugger->FilteredBuses.Sort(&Debug::CompareNames<FControlBusDebugInfo>);
 
-			FilteredMixes = RefreshedFilteredMixes;
-			FilteredMixes.Sort(&Debug::CompareNames<FControlBusMixDebugInfo>);
+			ThisDebugger->FilteredMixes = RefreshedFilteredMixes;
+			ThisDebugger->FilteredMixes.Sort(&Debug::CompareNames<FControlBusMixDebugInfo>);
 
 			FGeneratorSortMap NewFilteredMap;
-			if (bShowGenerators)
+			if (ThisDebugger->bShowGenerators)
 			{
 				NewFilteredMap = RefreshedFilteredGenerators;
 				for (TPair<FString, FGeneratorDebugInfo>& NewInfoPair : NewFilteredMap)
 				{
-					FGeneratorDebugInfo* LastDebugInfo = FilteredGeneratorsMap.Find(NewInfoPair.Key);
-					NewInfoPair.Value.bEnabled = bEnableAllGenerators
+					FGeneratorDebugInfo* LastDebugInfo = ThisDebugger->FilteredGeneratorsMap.Find(NewInfoPair.Key);
+					NewInfoPair.Value.bEnabled = ThisDebugger->bEnableAllGenerators
 						|| (LastDebugInfo && LastDebugInfo->bEnabled);
 
-					if (bool* Value = RequestedGeneratorUpdate.Find(NewInfoPair.Key))
+					if (bool* Value = ThisDebugger->RequestedGeneratorUpdate.Find(NewInfoPair.Key))
 					{
 						NewInfoPair.Value.bEnabled = *Value;
 					}
@@ -492,11 +498,11 @@ namespace AudioModulation
 						NewInfoPair.Value.FilteredInstances.Reset();
 					}
 				}
-				bEnableAllGenerators = 0;
-				RequestedGeneratorUpdate.Reset();
+				ThisDebugger->bEnableAllGenerators = 0;
+				ThisDebugger->RequestedGeneratorUpdate.Reset();
 			}
 
-			FilteredGeneratorsMap = MoveTemp(NewFilteredMap);
+			ThisDebugger->FilteredGeneratorsMap = MoveTemp(NewFilteredMap);
 		}, TStatId(), nullptr, ENamedThreads::GameThread);
 	}
 
diff --git a/Engine/Plugins/Runtime/AudioModulation/Source/AudioModulation/Private/AudioModulationDebugger.h b/Engine/Plugins/Runtime/AudioModulation/Source/AudioModulation/Private/AudioModulationDebugger.h
--- a/Engine/Plugins/Runtime/AudioModulation/Source/AudioModulation/Private/AudioModulationDebugger.h
+++ b/Engine/Plugins/Runtime/AudioModulation/Source/AudioModulation/Private/AudioModulationDebugger.h
@@ -1,7 +1,7 @@
 // Copyright Epic Games, Inc. All Rights Reserved.
 #pragma once
 
-#include "CoreMinimal.h"
+#include "Misc/Build.h"
 
 #if WITH_AUDIOMODULATION
 #if !UE_BUILD_SHIPPING
@@ -9,6 +9,7 @@
 #include "SoundControlBusMix.h"
 #include "SoundModulationGenerator.h"
 #include "SoundModulationProxy.h"
+#include "Templates/SharedPointer.h"
 
 
 namespace AudioModulation
@@ -70,7 +71,7 @@ namespace AudioModulation
 		TArray<FInstanceValues> FilteredInstances;
 	};
 
-	class FAudioModulationDebugger
+	class FAudioModulationDebugger : public TSharedFromThis<FAudioModulationDebugger, ESPMode::ThreadSafe>
 	{
 	public:
 		FAudioModulationDebugger();
diff --git a/Engine/Plugins/Runtime/AudioModulation/Source/AudioModulation/Private/AudioModulationSystem.cpp b/Engine/Plugins/Runtime/AudioModulation/Source/AudioModulation/Private/AudioModulationSystem.cpp
--- a/Engine/Plugins/Runtime/AudioModulation/Source/AudioModulation/Private/AudioModulationSystem.cpp
+++ b/Engine/Plugins/Runtime/AudioModulation/Source/AudioModulation/Private/AudioModulationSystem.cpp
@@ -47,6 +47,10 @@ namespace AudioModulation
 
 	void FAudioModulationSystem::Initialize(const FAudioPluginInitializationParams& InitializationParams)
 	{
+#if !UE_BUILD_SHIPPING
+		Debugger = MakeShared<FAudioModulationDebugger>();
+#endif // !UE_BUILD_SHIPPING
+
 		AudioDeviceId = InitializationParams.AudioDevicePtr->DeviceID;
 	}
 
@@ -59,19 +63,19 @@ namespace AudioModulation
 	bool FAudioModulationSystem::OnPostHelp(FCommonViewportClient* ViewportClient, const TCHAR* Stream)
 	{
 		check(IsInGameThread());
-		return ViewportClient ? Debugger.OnPostHelp(*ViewportClient, Stream) : true;
+		return ViewportClient ? Debugger->OnPostHelp(*ViewportClient, Stream) : true;
 	}
 
 	int32 FAudioModulationSystem::OnRenderStat(FViewport* Viewport, FCanvas* Canvas, int32 X, int32 Y, const UFont& Font, const FVector* ViewLocation, const FRotator* ViewRotation)
 	{
 		check(IsInGameThread());
-		return Canvas ? Debugger.OnRenderStat(*Canvas, X, Y, Font) : Y;
+		return Canvas ? Debugger->OnRenderStat(*Canvas, X, Y, Font) : Y;
 	}
 
 	bool FAudioModulationSystem::OnToggleStat(FCommonViewportClient* ViewportClient, const TCHAR* Stream)
 	{
 		check(IsInGameThread());
-		return ViewportClient ? Debugger.OnToggleStat(*ViewportClient, Stream) : true;
+		return ViewportClient ? Debugger->OnToggleStat(*ViewportClient, Stream) : true;
 	}
 #endif // !UE_BUILD_SHIPPING
 
@@ -193,32 +197,32 @@ namespace AudioModulation
 #if !UE_BUILD_SHIPPING
 	void FAudioModulationSystem::SetDebugBusFilter(const FString* InFilter)
 	{
-		Debugger.SetDebugBusFilter(InFilter);
+		Debugger->SetDebugBusFilter(InFilter);
 	}
 
 	void FAudioModulationSystem::SetDebugGeneratorsEnabled(bool bInIsEnabled)
 	{
-		Debugger.SetDebugGeneratorsEnabled(bInIsEnabled);
+		Debugger->SetDebugGeneratorsEnabled(bInIsEnabled);
 	}
 
 	void FAudioModulationSystem::SetDebugGeneratorFilter(const FString* InFilter)
 	{
-		Debugger.SetDebugGeneratorFilter(InFilter);
+		Debugger->SetDebugGeneratorFilter(InFilter);
 	}
 
 	void FAudioModulationSystem::SetDebugGeneratorTypeFilter(const FString* InFilter, bool bInEnabled)
 	{
-		Debugger.SetDebugGeneratorTypeFilter(InFilter, bInEnabled);
+		Debugger->SetDebugGeneratorTypeFilter(InFilter, bInEnabled);
 	}
 
 	void FAudioModulationSystem::SetDebugMatrixEnabled(bool bInIsEnabled)
 	{
-		Debugger.SetDebugMatrixEnabled(bInIsEnabled);
+		Debugger->SetDebugMatrixEnabled(bInIsEnabled);
 	}
 
 	void FAudioModulationSystem::SetDebugMixFilter(const FString* InNameFilter)
 	{
-		Debugger.SetDebugMixFilter(InNameFilter);
+		Debugger->SetDebugMixFilter(InNameFilter);
 	}
 #endif // !UE_BUILD_SHIPPING
 
@@ -489,7 +493,7 @@ namespace AudioModulation
 		SET_DWORD_STAT(STAT_AudioModulationProcQueueCount, CommandsProcessed);
 
 #if !UE_BUILD_SHIPPING
- 		Debugger.UpdateDebugData(InElapsed, RefProxies);
+ 		Debugger->UpdateDebugData(InElapsed, RefProxies);
 #endif // !UE_BUILD_SHIPPING
 	}
 
diff --git a/Engine/Plugins/Runtime/AudioModulation/Source/AudioModulation/Private/AudioModulationSystem.h b/Engine/Plugins/Runtime/AudioModulation/Source/AudioModulation/Private/AudioModulationSystem.h
--- a/Engine/Plugins/Runtime/AudioModulation/Source/AudioModulation/Private/AudioModulationSystem.h
+++ b/Engine/Plugins/Runtime/AudioModulation/Source/AudioModulation/Private/AudioModulationSystem.h
@@ -248,7 +248,7 @@ namespace AudioModulation
 		bool OnToggleStat(FCommonViewportClient* ViewportClient, const TCHAR* Stream);
 
 	private:
-		FAudioModulationDebugger Debugger;
+		TSharedPtr<FAudioModulationDebugger> Debugger;
 #endif // !UE_BUILD_SHIPPING
 
 		friend FControlBusProxy;
