Fix bug in Control Bus Mix Stage UI

diff --git a/Engine/Plugins/Runtime/AudioModulation/Source/AudioModulationEditor/Private/Layouts/SoundControlBusMixStageLayout.cpp b/Engine/Plugins/Runtime/AudioModulation/Source/AudioModulationEditor/Private/Layouts/SoundControlBusMixStageLayout.cpp
--- a/Engine/Plugins/Runtime/AudioModulation/Source/AudioModulationEditor/Private/Layouts/SoundControlBusMixStageLayout.cpp
+++ b/Engine/Plugins/Runtime/AudioModulation/Source/AudioModulationEditor/Private/Layouts/SoundControlBusMixStageLayout.cpp
@@ -263,31 +263,31 @@ void FSoundControlBusMixStageLayoutCustomization::CustomizeChildren(TSharedRef<I
 						return EVisibility::Visible;
 					}))
 				]
-				+ SHorizontalBox::Slot()
-					.FillWidth(0.4f)
-					.Padding(4.0f, 0.0f, 0.0f, 0.0f)
-					.VAlign(VAlign_Center)
-					[
-						NormalizedValueWidget
-					]
-				+ SHorizontalBox::Slot()
-					.Padding(4.0f, 0.0f, 0.0f, 0.0f)
-					.VAlign(VAlign_Center)
-					.AutoWidth()
-					[
-						SNew(STextBlock)
-						.Font(IDetailLayoutBuilder::GetDetailFont())
-						.Text(LOCTEXT("SoundModulationControl_UnitValueNormalized", "Normalized"))
-						.Visibility(TAttribute<EVisibility>::Create([BusHandle]()
+			+ SHorizontalBox::Slot()
+				.FillWidth(0.4f)
+				.Padding(4.0f, 0.0f, 0.0f, 0.0f)
+				.VAlign(VAlign_Center)
+				[
+					NormalizedValueWidget
+				]
+			+ SHorizontalBox::Slot()
+				.Padding(4.0f, 0.0f, 0.0f, 0.0f)
+				.VAlign(VAlign_Center)
+				.AutoWidth()
+				[
+					SNew(STextBlock)
+					.Font(IDetailLayoutBuilder::GetDetailFont())
+					.Text(LOCTEXT("SoundModulationControl_UnitValueNormalized", "Normalized"))
+					.Visibility(TAttribute<EVisibility>::Create([BusHandle]()
+					{
+						if (USoundModulationParameter* Parameter = AudioModulationEditorUtils::GetParameterFromBus(BusHandle))
 						{
-							if (USoundModulationParameter* Parameter = AudioModulationEditorUtils::GetParameterFromBus(BusHandle))
-							{
-								return Parameter->RequiresUnitConversion() ? EVisibility::Visible : EVisibility::Hidden;
-							}
+							return Parameter->RequiresUnitConversion() ? EVisibility::Visible : EVisibility::Hidden;
+						}
 
-							return EVisibility::Visible;
-						}))
-					]
+						return EVisibility::Hidden;
+					}))
+				]
 		]
 	.OverrideResetToDefault(ResetToDefault);
 
