Update Audio Modulation cvars to use "Modulation" instead of "SoundModulators"
Also update the modulation map to use the new cvar

diff --git a/Engine/Plugins/Runtime/AudioModulation/Source/AudioModulation/Private/AudioModulationDebugger.cpp b/Engine/Plugins/Runtime/AudioModulation/Source/AudioModulation/Private/AudioModulationDebugger.cpp
--- a/Engine/Plugins/Runtime/AudioModulation/Source/AudioModulation/Private/AudioModulationDebugger.cpp
+++ b/Engine/Plugins/Runtime/AudioModulation/Source/AudioModulation/Private/AudioModulationDebugger.cpp
@@ -44,7 +44,7 @@ namespace AudioModulation
 			{
 				if (GEngine)
 				{
-					GEngine->Exec(World, TEXT("au.Debug.SoundModulators 1"));
+					GEngine->Exec(World, TEXT("au.Debug.Modulation 1"));
 				}
 
 				if (IAudioModulationManager* Modulation = DeviceHandle->ModulationInterface.Get())
@@ -59,14 +59,14 @@ namespace AudioModulation
 
 static float AudioModulationDebugUpdateRateCVar = 0.1f;
 FAutoConsoleVariableRef CVarAudioModulationDebugUpdateRate(
-	TEXT("au.Debug.SoundModulators.UpdateRate"),
+	TEXT("au.Debug.Modulation.UpdateRate"),
 	AudioModulationDebugUpdateRateCVar,
 	TEXT("Sets update rate for modulation debug statistics (in seconds).\n")
 	TEXT("Default: 0.1f"),
 	ECVF_Default);
 
 static FAutoConsoleCommandWithWorldAndArgs CVarAudioModulationDebugBuses(
-	TEXT("au.Debug.SoundModulators.Filter.Buses"),
+	TEXT("au.Debug.Modulation.Filter.Buses"),
 	TEXT("Sets substring by which to filter mixes in matrix view. Arguments:\n"
 		"Enabled (Optional, ex. True, False. Default: True) - Whether or not to enable showing buses.\n"
 		"Filter (Optional, Default: null) - Whether or not to filter buses by name using the provided substring.\n"),
@@ -80,7 +80,7 @@ static FAutoConsoleCommandWithWorldAndArgs CVarAudioModulationDebugBuses(
 	}));
 
 static FAutoConsoleCommandWithWorldAndArgs CVarAudioModulationDebugMixes(
-	TEXT("au.Debug.SoundModulators.Filter.Mixes"),
+	TEXT("au.Debug.Modulation.Filter.Mixes"),
 	TEXT("Sets substring by which to filter mixes in matrix view. Arguments:\n"
 		"Filter - Filter bus mixes by name using the provided substring.\n"),
 	FConsoleCommandWithWorldAndArgsDelegate::CreateLambda([](const TArray<FString>& Args, UWorld* World)
@@ -93,7 +93,7 @@ static FAutoConsoleCommandWithWorldAndArgs CVarAudioModulationDebugMixes(
 	}));
 
 static FAutoConsoleCommandWithWorldAndArgs CVarAudioModulationDebugMatrix(
-	TEXT("au.Debug.SoundModulators.Enable.Matrix"),
+	TEXT("au.Debug.Modulation.Enable.Matrix"),
 	TEXT("Whether or not to enable mix matrix. Arguments:\n"
 		"Enabled (Default: 1/True) - Whether or not to have matrix view enabled.\n"),
 	FConsoleCommandWithWorldAndArgsDelegate::CreateLambda([](const TArray<FString>& Args, UWorld* World)
@@ -110,7 +110,7 @@ static FAutoConsoleCommandWithWorldAndArgs CVarAudioModulationDebugMatrix(
 	}));
 
 static FAutoConsoleCommandWithWorldAndArgs CVarAudioModulationSetDebugGeneratorFilter(
-	TEXT("au.Debug.SoundModulators.Filter.Generators"),
+	TEXT("au.Debug.Modulation.Filter.Generators"),
 	TEXT("Sets substring by which to filter generators. Arguments:\n"
 		"Name - Filter generators by name using the provided substring.\n"),
 	FConsoleCommandWithWorldAndArgsDelegate::CreateLambda([](const TArray<FString>& Args, UWorld* World)
@@ -124,7 +124,7 @@ static FAutoConsoleCommandWithWorldAndArgs CVarAudioModulationSetDebugGeneratorF
 	}));
 
 static FAutoConsoleCommandWithWorldAndArgs CVarAudioModulationSetDebugGeneratorFilterType(
-	TEXT("au.Debug.SoundModulators.Filter.Generators.Type"),
+	TEXT("au.Debug.Modulation.Filter.Generators.Type"),
 	TEXT("Whether to display or hide Generator type provided (defaults to show if enablement boolean not provided). Arguments:\n"
 		"Name (Optional, Default: null) - Filter generators type to display/hide (Empty/null clears any currently set filter).\n"
 		"Enabled (Optional, Default: True) - True to show, false to hide.\n"),
@@ -144,7 +144,7 @@ static FAutoConsoleCommandWithWorldAndArgs CVarAudioModulationSetDebugGeneratorF
 	}));
 
 static FAutoConsoleCommandWithWorldAndArgs CVarAudioModulationDebugEnableGenerators(
-	TEXT("au.Debug.SoundModulators.Enable.Generators"),
+	TEXT("au.Debug.Modulation.Enable.Generators"),
 	TEXT("Whether or not to enable displaying generators. Arguments:\n"
 		"Enabled (Default: 1/True) - Whether or not to have generators view enabled.\n"),
 	FConsoleCommandWithWorldAndArgsDelegate::CreateLambda([](const TArray<FString>& Args, UWorld* World)
diff --git a/Engine/Source/Runtime/Engine/Private/AudioDebug.cpp b/Engine/Source/Runtime/Engine/Private/AudioDebug.cpp
--- a/Engine/Source/Runtime/Engine/Private/AudioDebug.cpp
+++ b/Engine/Source/Runtime/Engine/Private/AudioDebug.cpp
@@ -786,8 +786,8 @@ static FAutoConsoleCommandWithWorldAndArgs GAudioDebugSoundReverb
 
 static FAutoConsoleCommandWithWorldAndArgs GAudioDebugSoundModulation
 (
-	TEXT("au.Debug.SoundModulators"),
-	TEXT("Post SoundModulation information to viewport(s).\n")
+	TEXT("au.Debug.Modulation"),
+	TEXT("Post Audio Modulation information to viewport(s).\n")
 		TEXT("0: Disable, 1: Enable\n")
 		TEXT("(Optional) -AllViews: Enables/Disables for all viewports, not just those associated with the current world"),
 	FConsoleCommandWithWorldAndArgsDelegate::CreateStatic([](const TArray<FString>& Args, UWorld* InWorld)
