Update Audio Modulation static BP functions to use "Modulation" category

diff --git a/Engine/Plugins/Runtime/AudioModulation/Source/AudioModulation/Public/AudioModulationStatics.h b/Engine/Plugins/Runtime/AudioModulation/Source/AudioModulation/Public/AudioModulationStatics.h
--- a/Engine/Plugins/Runtime/AudioModulation/Source/AudioModulation/Public/AudioModulationStatics.h
+++ b/Engine/Plugins/Runtime/AudioModulation/Source/AudioModulation/Public/AudioModulationStatics.h
@@ -38,7 +38,7 @@ public:
 	 * if bus is manually deactivated or destroyed (i.e. will not deactivate
 	 * when all references become inactive).
 	 */
-	UFUNCTION(BlueprintCallable, Category = "Audio", DisplayName = "Activate Control Bus", meta = (
+	UFUNCTION(BlueprintCallable, Category = "Audio|Modulation", DisplayName = "Activate Control Bus", meta = (
 		WorldContext = "WorldContextObject", 
 		Keywords = "modulation modulator")
 	)
@@ -49,7 +49,7 @@ public:
 	 * when all references become inactive).
 	 * @param BusMix - Mix to activate
 	 */
-	UFUNCTION(BlueprintCallable, Category = "Audio", DisplayName = "Activate Control Bus Mix", meta = (
+	UFUNCTION(BlueprintCallable, Category = "Audio|Modulation", DisplayName = "Activate Control Bus Mix", meta = (
 		WorldContext = "WorldContextObject", 
 		Keywords = "modulation modulator")
 	)
@@ -60,7 +60,7 @@ public:
 	 * when all references become inactive).
 	 * @param Modulator - Modulator to activate
 	 */
-	UFUNCTION(BlueprintCallable, Category = "Audio", DisplayName = "Activate Modulation Generator", meta = (
+	UFUNCTION(BlueprintCallable, Category = "Audio|Modulation", DisplayName = "Activate Modulation Generator", meta = (
 		WorldContext = "WorldContextObject", 
 		Keywords = "modulator lfo envelope follower")
 	)
@@ -74,7 +74,7 @@ public:
 	 * when all references become inactive).
 	 * @return Capture this in a Blueprint variable to prevent it from being automatically garbage collected. 
 	 */
-	UFUNCTION(BlueprintCallable, Category = "Audio", DisplayName = "Create Control Bus", meta = (
+	UFUNCTION(BlueprintCallable, Category = "Audio|Modulation", DisplayName = "Create Control Bus", meta = (
 		AdvancedDisplay = "3",
 		WorldContext = "WorldContextObject",
 		Keywords = "make modulation LPF modulator")
@@ -87,7 +87,7 @@ public:
 	 * @param AttackTime - Time in seconds for stage to mix in.
 	 * @param ReleaseTime - Time in seconds for stage to mix out.
 	 */
-	UFUNCTION(BlueprintCallable, Category = "Audio", DisplayName = "Create Control Bus Mix Stage", meta = (
+	UFUNCTION(BlueprintCallable, Category = "Audio|Modulation", DisplayName = "Create Control Bus Mix Stage", meta = (
 		AdvancedDisplay = "3",
 		WorldContext = "WorldContextObject",
 		Keywords = "make modulation modulator stage")
@@ -107,7 +107,7 @@ public:
 	 * when all references become inactive).
 	 * @return Capture this in a Blueprint variable to prevent it from being automatically garbage collected. 
 	 */
-	UFUNCTION(BlueprintCallable, Category = "Audio", DisplayName = "Create Control Bus Mix", meta = (
+	UFUNCTION(BlueprintCallable, Category = "Audio|Modulation", DisplayName = "Create Control Bus Mix", meta = (
 		WorldContext = "WorldContextObject",
 		Keywords = "make modulation modulator")
 	)
@@ -120,7 +120,7 @@ public:
 	/** Deactivates a bus. Does nothing if the provided bus is already inactive.
 	 * @param Bus - Scope of modulator
 	 */
-	UFUNCTION(BlueprintCallable, Category = "Audio", DisplayName = "Deactivate Control Bus", meta = (
+	UFUNCTION(BlueprintCallable, Category = "Audio|Modulation", DisplayName = "Deactivate Control Bus", meta = (
 		WorldContext = "WorldContextObject", 
 		Keywords = "modulation modulator bus")
 	)
@@ -129,7 +129,7 @@ public:
 	/** Deactivates a modulation bus mix. Does nothing if an instance of the provided bus mix is already inactive.
 	 * @param BusMix - Mix to deactivate
 	 */
-	UFUNCTION(BlueprintCallable, Category = "Audio", DisplayName = "Deactivate Control Bus Mix", meta = (
+	UFUNCTION(BlueprintCallable, Category = "Audio|Modulation", DisplayName = "Deactivate Control Bus Mix", meta = (
 		WorldContext = "WorldContextObject",
 		Keywords = "modulation modulator")
 	)
@@ -139,7 +139,7 @@ public:
 	 * @param Generator - Generator to activate
 	 * @param Scope - Scope of modulator
 	 */
-	UFUNCTION(BlueprintCallable, Category = "Audio", DisplayName = "Deactivate Modulation Generator", meta = (
+	UFUNCTION(BlueprintCallable, Category = "Audio|Modulation", DisplayName = "Deactivate Modulation Generator", meta = (
 		WorldContext = "WorldContextObject",
 		Keywords = "bus modulation modulator generator")
 	)
@@ -150,7 +150,7 @@ public:
 	 * @param BusMix - Mix object to serialize to profile .ini.
 	 * @param ProfileIndex - Index of profile, allowing multiple profiles can be saved for single mix object. If 0, saves to default ini profile (no suffix).
 	 */
-	UFUNCTION(BlueprintCallable, Category = "Audio", DisplayName = "Save Control Bus Mix to Profile", meta = (
+	UFUNCTION(BlueprintCallable, Category = "Audio|Modulation", DisplayName = "Save Control Bus Mix to Profile", meta = (
 		WorldContext = "WorldContextObject",
 		AdvancedDisplay = "2",
 		Keywords = "serialize modulation modulator ini")
@@ -163,7 +163,7 @@ public:
 	 * @param ProfileIndex - Index of profile, allowing multiple profiles to be loaded to single mix object. If <= 0, loads from default profile (no suffix).
 	 * @return Stages - Stage values loaded from profile (empty if profile did not exist or had no values serialized).
 	 */
-	UFUNCTION(BlueprintCallable, Category = "Audio", DisplayName = "Load Control Bus Mix From Profile", meta = (
+	UFUNCTION(BlueprintCallable, Category = "Audio|Modulation", DisplayName = "Load Control Bus Mix From Profile", meta = (
 		WorldContext = "WorldContextObject",
 		AdvancedDisplay = "2",
 		Keywords = "deserialize modulation modulator ini")
@@ -179,7 +179,7 @@ public:
 	 *					 If negative, falls back to last fade time set on stage. If fade time never set on stage,
 	 *					 uses attack time set on stage in mix asset.
 	 */
-	UFUNCTION(BlueprintCallable, Category = "Audio", DisplayName = "Set Control Bus Mix", meta = (
+	UFUNCTION(BlueprintCallable, Category = "Audio|Modulation", DisplayName = "Set Control Bus Mix", meta = (
 		WorldContext = "WorldContextObject",
 		Keywords = "modulation modulator stage")
 	)
@@ -193,7 +193,7 @@ public:
 	 * @param FadeTime - Fade time to user when interpolating between current value and new value. If negative, falls back to last fade
 	 * time set on stage. If fade time never set on stage, defaults to 100ms.
 	 */
-	UFUNCTION(BlueprintCallable, Category = "Audio", DisplayName = "Set Global Control Bus Mix Value", meta = (
+	UFUNCTION(BlueprintCallable, Category = "Audio|Modulation", DisplayName = "Set Global Control Bus Mix Value", meta = (
 		WorldContext = "WorldContextObject",
 		Keywords = "modulation modulator stage")
 	)
@@ -204,7 +204,7 @@ public:
 	 * @param FadeTime - Fade time to user when interpolating between current value and new values.
 	 *					 If non-positive, change is immediate.
 	 */
-	UFUNCTION(BlueprintCallable, Category = "Audio", DisplayName = "Clear Global Control Bus Mix Value", meta = (
+	UFUNCTION(BlueprintCallable, Category = "Audio|Modulation", DisplayName = "Clear Global Control Bus Mix Value", meta = (
 		WorldContext = "WorldContextObject",
 		Keywords = "modulation modulator stage")
 	)
@@ -214,7 +214,7 @@ public:
 	 * @param FadeTime - Fade time to user when interpolating between current value and new values.
 	 *					 If non-positive, change is immediate.
 	 */
-	UFUNCTION(BlueprintCallable, Category = "Audio", DisplayName = "Clear All Global Control Bus Mix Values", meta = (
+	UFUNCTION(BlueprintCallable, Category = "Audio|Modulation", DisplayName = "Clear All Global Control Bus Mix Values", meta = (
 		WorldContext = "WorldContextObject",
 		Keywords = "modulation modulator stage")
 	)
@@ -229,7 +229,7 @@ public:
 	 * @param Value - Target value to mix filtered stages to.
 	 * @param FadeTime - If non-negative, updates the fade time for the resulting bus stages found matching the provided filter.
 	 */
-	UFUNCTION(BlueprintCallable, Category = "Audio", DisplayName = "Set Control Bus Mix By Filter", meta = (
+	UFUNCTION(BlueprintCallable, Category = "Audio|Modulation", DisplayName = "Set Control Bus Mix By Filter", meta = (
 		AdvancedDisplay = "6",
 		WorldContext = "WorldContextObject",
 		Keywords = "control class modulation modulator stage value")
@@ -250,7 +250,7 @@ public:
 	 *					 If negative, falls back to last fade time set on stage. If fade time never set on stage,
 	 *					 uses attack time set on stage in mix asset.
 	 */
-	UFUNCTION(BlueprintCallable, Category = "Audio", DisplayName = "Update Control Bus Mix", meta = (
+	UFUNCTION(BlueprintCallable, Category = "Audio|Modulation", DisplayName = "Update Control Bus Mix", meta = (
 		WorldContext = "WorldContextObject",
 		Keywords = "set modulation modulator")
 	)
@@ -260,7 +260,7 @@ public:
 	 *  to active instance in audio thread (ignored if modulator type has not been activated).
 	 * @param Modulator - Modulator to update
 	 */
-	UFUNCTION(BlueprintCallable, Category = "Audio", DisplayName = "Update Modulator", meta = (
+	UFUNCTION(BlueprintCallable, Category = "Audio|Modulation", DisplayName = "Update Modulator", meta = (
 		WorldContext = "WorldContextObject",
 		Keywords = "set control bus mix modulation modulator generator")
 	)
