New Audio Modulation BP Nodes (Create Parameter, Create Env Follower Generator, Create LFO Generator)

diff --git a/Engine/Plugins/Runtime/AudioModulation/Source/AudioModulation/Private/AudioModulationStatics.cpp b/Engine/Plugins/Runtime/AudioModulation/Source/AudioModulation/Private/AudioModulationStatics.cpp
--- a/Engine/Plugins/Runtime/AudioModulation/Source/AudioModulation/Private/AudioModulationStatics.cpp
+++ b/Engine/Plugins/Runtime/AudioModulation/Source/AudioModulation/Private/AudioModulationStatics.cpp
@@ -256,6 +256,60 @@ USoundControlBusMix* UAudioModulationStatics::CreateBusMix(UObject* WorldContext
 	return NewBusMix;
 }
 
+USoundModulationParameter* UAudioModulationStatics::CreateModulationParameter(UObject* WorldContextObject, FName Name, TSubclassOf<USoundModulationParameter> ParamClass, float DefaultValue)
+{
+	UWorld* World = GetAudioWorld(WorldContextObject);
+	if (!World)
+	{
+		return nullptr;
+	}
+
+	USoundModulationParameter* NewParam = NewObject<USoundModulationParameter>(WorldContextObject, ParamClass, Name);
+
+	if (NewParam)
+	{
+		NewParam->Settings.ValueNormalized = DefaultValue;
+	}
+
+	return NewParam;
+}
+
+USoundModulationGeneratorEnvelopeFollower* UAudioModulationStatics::CreateEnvelopeFollowerGenerator(UObject* WorldContextObject, FName Name, FEnvelopeFollowerGeneratorParams Params)
+{
+	UWorld* World = GetAudioWorld(WorldContextObject);
+	if (!World)
+	{
+		return nullptr;
+	}
+	
+	USoundModulationGeneratorEnvelopeFollower* NewGenerator = NewObject<USoundModulationGeneratorEnvelopeFollower>(WorldContextObject, USoundModulationGeneratorEnvelopeFollower::StaticClass(), Name);
+
+	if (NewGenerator)
+	{
+		NewGenerator->Params = Params;
+	}
+
+	return NewGenerator;
+}
+
+USoundModulationGeneratorLFO* UAudioModulationStatics::CreateLFOGenerator(UObject* WorldContextObject, FName Name, FSoundModulationLFOParams Params)
+{
+	UWorld* World = GetAudioWorld(WorldContextObject);
+	if (!World)
+	{
+		return nullptr;
+	}
+
+	USoundModulationGeneratorLFO* NewGenerator = NewObject<USoundModulationGeneratorLFO>(WorldContextObject, USoundModulationGeneratorLFO::StaticClass(), Name);
+
+	if (NewGenerator)
+	{
+		NewGenerator->Params = Params;
+	}
+
+	return NewGenerator;
+}
+
 void UAudioModulationStatics::DeactivateBus(const UObject* WorldContextObject, USoundControlBus* Bus)
 {
 	if (Bus)
diff --git a/Engine/Plugins/Runtime/AudioModulation/Source/AudioModulation/Public/AudioModulationStatics.h b/Engine/Plugins/Runtime/AudioModulation/Source/AudioModulation/Public/AudioModulationStatics.h
--- a/Engine/Plugins/Runtime/AudioModulation/Source/AudioModulation/Public/AudioModulationStatics.h
+++ b/Engine/Plugins/Runtime/AudioModulation/Source/AudioModulation/Public/AudioModulationStatics.h
@@ -1,6 +1,8 @@
 // Copyright Epic Games, Inc. All Rights Reserved.
 #pragma once
 
+#include "Generators/SoundModulationEnvelopeFollower.h"
+#include "Generators/SoundModulationLFO.h"
 #include "Kismet/BlueprintFunctionLibrary.h"
 #include "SoundControlBus.h"
 #include "SoundControlBusMix.h"
@@ -92,7 +94,7 @@ public:
 		WorldContext = "WorldContextObject",
 		Keywords = "make modulation modulator stage")
 	)
-	static FSoundControlBusMixStage CreateBusMixStage(
+	static UPARAM(DisplayName = "Stage") FSoundControlBusMixStage CreateBusMixStage(
 		const UObject* WorldContextObject,
 		USoundControlBus* Bus,
 		float Value,
@@ -105,7 +107,7 @@ public:
 	 * @param Activate - Whether or not to activate mix on creation. If true, deactivation will only occur
 	 * if returned mix is manually deactivated and not referenced or destroyed (i.e. will not deactivate
 	 * when all references become inactive).
-	 * @return Capture this in a Blueprint variable to prevent it from being automatically garbage collected. 
+	 * @return Capture this in a Blueprint variable to prevent it from being garbage collected. 
 	 */
 	UFUNCTION(BlueprintCallable, Category = "Audio|Modulation", DisplayName = "Create Control Bus Mix", meta = (
 		WorldContext = "WorldContextObject",
@@ -117,6 +119,50 @@ public:
 		TArray<FSoundControlBusMixStage> Stages,
 		bool Activate);
 
+	/** Creates a modulation parameter of a given class.
+	 * @param Name - Name of parameter.
+	 * @param ParamClass - The type of Modulation Parameter to create.
+	 * @param DefaultValue - The default normalized value of the parameter (range 0-1).
+	 * @return Capture this in a Blueprint variable to prevent it from being garbage collected.
+	 */
+	UFUNCTION(BlueprintCallable, Category = "Audio|Modulation", DisplayName = "Create Modulation Parameter", meta = (
+		WorldContext = "WorldContextObject",
+		Keywords = "make modulator")
+	)
+		static UPARAM(DisplayName = "Parameter") USoundModulationParameter* CreateModulationParameter(
+			UObject* WorldContextObject,
+			FName Name,
+			TSubclassOf<USoundModulationParameter> ParamClass,
+			float DefaultValue);
+
+	/** Creates a modulation generator based on an Envelope Follower with the given parameters.
+	 * @param Name - Name of generator.
+	 * @param Settings - The Envelope Follower settings, including what data to follow.
+	 * @return Capture this in a Blueprint variable to prevent it from being garbage collected.
+	 */
+	UFUNCTION(BlueprintCallable, Category = "Audio|Modulation", DisplayName = "Create Envelope Follower Generator", meta = (
+		WorldContext = "WorldContextObject",
+		Keywords = "make modulator")
+	)
+		static UPARAM(DisplayName = "Generator") USoundModulationGeneratorEnvelopeFollower* CreateEnvelopeFollowerGenerator(
+			UObject* WorldContextObject,
+			FName Name,
+			FEnvelopeFollowerGeneratorParams Params);
+
+	/** Creates a modulation generator based on an LFO with the given parameters.
+	 * @param Name - Name of generator.
+	 * @param Settings - The LFO Settings.
+	 * @return Capture this in a Blueprint variable to prevent it from being garbage collected.
+	 */
+	UFUNCTION(BlueprintCallable, Category = "Audio|Modulation", DisplayName = "Create LFO Generator", meta = (
+		WorldContext = "WorldContextObject",
+		Keywords = "make modulator")
+	)
+		static UPARAM(DisplayName = "Generator") USoundModulationGeneratorLFO* CreateLFOGenerator(
+			UObject* WorldContextObject,
+			FName Name,
+			FSoundModulationLFOParams Params);
+
 	/** Deactivates a bus. Does nothing if the provided bus is already inactive.
 	 * @param Bus - Scope of modulator
 	 */
