Get Modulator Values in blueprint

diff --git a/Engine/Plugins/Runtime/AudioModulation/Source/AudioModulation/Private/AudioModulation.cpp b/Engine/Plugins/Runtime/AudioModulation/Source/AudioModulation/Private/AudioModulation.cpp
--- a/Engine/Plugins/Runtime/AudioModulation/Source/AudioModulation/Private/AudioModulation.cpp
+++ b/Engine/Plugins/Runtime/AudioModulation/Source/AudioModulation/Private/AudioModulation.cpp
@@ -159,6 +159,13 @@ namespace AudioModulation
 		ModSystem->ClearAllGlobalBusMixValues(InFadeTime);
 	}
 
+	float FAudioModulationManager::GetModulatorValueThreadSafe(uint32 ModulationID)
+	{
+		float OutValue = 1.0f;
+		ModSystem->GetModulatorValueThreadSafe(ModulationID, OutValue);
+		return OutValue;
+	}
+
 #if !UE_BUILD_SHIPPING
 	bool FAudioModulationManager::OnPostHelp(FCommonViewportClient* ViewportClient, const TCHAR* Stream)
 	{
diff --git a/Engine/Plugins/Runtime/AudioModulation/Source/AudioModulation/Private/AudioModulationStatics.cpp b/Engine/Plugins/Runtime/AudioModulation/Source/AudioModulation/Private/AudioModulationStatics.cpp
--- a/Engine/Plugins/Runtime/AudioModulation/Source/AudioModulation/Private/AudioModulationStatics.cpp
+++ b/Engine/Plugins/Runtime/AudioModulation/Source/AudioModulation/Private/AudioModulationStatics.cpp
@@ -467,5 +467,21 @@ void UAudioModulationStatics::UpdateModulator(const UObject* WorldContextObject,
 		}
 	}
 }
+
+float UAudioModulationStatics::GetModulatorValue(const UObject* WorldContextObject, USoundModulatorBase* Modulator)
+{	
+	if (Modulator)
+	{
+		UWorld* World = GetAudioWorld(WorldContextObject);
+		if (AudioModulation::FAudioModulationManager* ModSystem = GetModulation(World))
+		{
+			return ModSystem->GetModulatorValueThreadSafe(Modulator->GetUniqueID());
+		}
+	}
+
+	return 1.0f;
+}
+
+
 #undef LOCTEXT_NAMESPACE
 
diff --git a/Engine/Plugins/Runtime/AudioModulation/Source/AudioModulation/Private/AudioModulationSystem.cpp b/Engine/Plugins/Runtime/AudioModulation/Source/AudioModulation/Private/AudioModulationSystem.cpp
--- a/Engine/Plugins/Runtime/AudioModulation/Source/AudioModulation/Private/AudioModulationSystem.cpp
+++ b/Engine/Plugins/Runtime/AudioModulation/Source/AudioModulation/Private/AudioModulationSystem.cpp
@@ -373,6 +373,19 @@ namespace AudioModulation
 		return false;
 	}
 
+	bool FAudioModulationSystem::GetModulatorValueThreadSafe(uint32 ModulatorID, float& OutValue) const
+	{
+		FScopeLock Lock(&ThreadSafeModValueCritSection);
+
+		if (const float* Value = ThreadSafeModValueMap.Find(ModulatorID))
+		{
+			OutValue = *Value;
+			return true;
+		}
+
+		return false;
+	}
+
 	Audio::FDeviceId FAudioModulationSystem::GetAudioDeviceId() const
 	{
 		return AudioDeviceId;
diff --git a/Engine/Plugins/Runtime/AudioModulation/Source/AudioModulation/Private/AudioModulationSystem.h b/Engine/Plugins/Runtime/AudioModulation/Source/AudioModulation/Private/AudioModulationSystem.h
--- a/Engine/Plugins/Runtime/AudioModulation/Source/AudioModulation/Private/AudioModulationSystem.h
+++ b/Engine/Plugins/Runtime/AudioModulation/Source/AudioModulation/Private/AudioModulationSystem.h
@@ -102,6 +102,9 @@ namespace AudioModulation
 		/* Attempts to get the modulator value from any thread (lock contentious).*/
 		bool GetModulatorValueThreadSafe(const Audio::FModulatorHandle& ModulatorHandle, float& OutValue) const;
 
+		/* Attempts to get the modulator value from any thread (lock contentious).*/
+		bool GetModulatorValueThreadSafe(uint32 ModulatorID, float& OutValue) const;
+
 		void UnregisterModulator(const Audio::FModulatorHandle& InHandle);
 
 		/* Saves mix to .ini profile for fast iterative development that does not require re-cooking a mix */
diff --git a/Engine/Plugins/Runtime/AudioModulation/Source/AudioModulation/Public/AudioModulation.h b/Engine/Plugins/Runtime/AudioModulation/Source/AudioModulation/Public/AudioModulation.h
--- a/Engine/Plugins/Runtime/AudioModulation/Source/AudioModulation/Public/AudioModulation.h
+++ b/Engine/Plugins/Runtime/AudioModulation/Source/AudioModulation/Public/AudioModulation.h
@@ -81,6 +81,8 @@ namespace AudioModulation
 		void ClearGlobalBusMixValue(const USoundControlBus& InBus, float InFadeTime);
 		void ClearAllGlobalBusMixValues(float InFadeTime);
 
+		float GetModulatorValueThreadSafe(uint32 ModulationID);
+
 		FAudioModulationSystem& GetSystem();
 
 	private:
diff --git a/Engine/Plugins/Runtime/AudioModulation/Source/AudioModulation/Public/AudioModulationStatics.h b/Engine/Plugins/Runtime/AudioModulation/Source/AudioModulation/Public/AudioModulationStatics.h
--- a/Engine/Plugins/Runtime/AudioModulation/Source/AudioModulation/Public/AudioModulationStatics.h
+++ b/Engine/Plugins/Runtime/AudioModulation/Source/AudioModulation/Public/AudioModulationStatics.h
@@ -311,4 +311,13 @@ public:
 		Keywords = "set control bus mix modulation modulator generator")
 	)
 	static void UpdateModulator(const UObject* WorldContextObject, USoundModulatorBase* Modulator);
+
+	/** Gets the (normalized) value of the given modulator. 
+	 * @return Value - The current value of the modulator. If the modulator is not active, returns 1.0.
+	 */
+	UFUNCTION(BlueprintCallable, Category = "Audio|Modulation", DisplayName = "Get Modulator Value", meta = (
+		WorldContext = "WorldContextObject",
+		Keywords = "modulation generator bus")
+	)
+	static UPARAM(DisplayName = "Value") float GetModulatorValue(const UObject* WorldContextObject, USoundModulatorBase* Modulator);
 };
