Remove exec pins from Get Modulator BP Function

diff --git a/Engine/Plugins/Runtime/AudioModulation/Source/AudioModulation/Public/AudioModulationStatics.h b/Engine/Plugins/Runtime/AudioModulation/Source/AudioModulation/Public/AudioModulationStatics.h
--- a/Engine/Plugins/Runtime/AudioModulation/Source/AudioModulation/Public/AudioModulationStatics.h
+++ b/Engine/Plugins/Runtime/AudioModulation/Source/AudioModulation/Public/AudioModulationStatics.h
@@ -315,7 +315,7 @@ public:
 	/** Gets the (normalized) value of the given modulator. 
 	 * @return Value - The current value of the modulator. If the modulator is not active, returns 1.0.
 	 */
-	UFUNCTION(BlueprintCallable, Category = "Audio|Modulation", DisplayName = "Get Modulator Value", meta = (
+	UFUNCTION(BlueprintPure, Category = "Audio|Modulation", DisplayName = "Get Modulator Value", meta = (
 		WorldContext = "WorldContextObject",
 		Keywords = "modulation generator bus")
 	)
