Hide "Create Modulation Parameter" BP Function

diff --git a/Engine/Plugins/Runtime/AudioModulation/Source/AudioModulation/Public/AudioModulationStatics.h b/Engine/Plugins/Runtime/AudioModulation/Source/AudioModulation/Public/AudioModulationStatics.h
--- a/Engine/Plugins/Runtime/AudioModulation/Source/AudioModulation/Public/AudioModulationStatics.h
+++ b/Engine/Plugins/Runtime/AudioModulation/Source/AudioModulation/Public/AudioModulationStatics.h
@@ -125,7 +125,7 @@ public:
 	 * @param DefaultValue - The default normalized value of the parameter (range 0-1).
 	 * @return Capture this in a Blueprint variable to prevent it from being garbage collected.
 	 */
-	UFUNCTION(BlueprintCallable, Category = "Audio|Modulation", DisplayName = "Create Modulation Parameter", meta = (
+	UFUNCTION(Category = "Audio|Modulation", DisplayName = "Create Modulation Parameter", meta = (
 		WorldContext = "WorldContextObject",
 		Keywords = "make modulator")
 	)
