Fix envelope follower modulation generator hang caused by accessing audio device during teardown

diff --git a/Engine/Plugins/Runtime/AudioModulation/Source/AudioModulation/Private/Generators/SoundModulationEnvelopeFollower.cpp b/Engine/Plugins/Runtime/AudioModulation/Source/AudioModulation/Private/Generators/SoundModulationEnvelopeFollower.cpp
--- a/Engine/Plugins/Runtime/AudioModulation/Source/AudioModulation/Private/Generators/SoundModulationEnvelopeFollower.cpp
+++ b/Engine/Plugins/Runtime/AudioModulation/Source/AudioModulation/Private/Generators/SoundModulationEnvelopeFollower.cpp
@@ -42,6 +42,7 @@ namespace AudioModulation
 			, bInvert(InGenerator.bInvert)
 			, bBusRequiresPatch(InGenerator.bBusRequiresPatch)
 		{
+			DeviceDestroyedHandle = FAudioDeviceManagerDelegates::OnAudioDeviceDestroyed.AddRaw(this, &FEnvelopeFollowerGenerator::OnDeviceDestroyed);
 		}
 
 		FEnvelopeFollowerGenerator(const FEnvelopeFollowerGeneratorParams& InGeneratorParams)
@@ -50,6 +51,8 @@ namespace AudioModulation
 			, bInvert(InGeneratorParams.bInvert ? 1 : 0)
 			, bBusRequiresPatch(1)
 		{
+			DeviceDestroyedHandle = FAudioDeviceManagerDelegates::OnAudioDeviceDestroyed.AddRaw(this, &FEnvelopeFollowerGenerator::OnDeviceDestroyed);
+
 			if (UAudioBus* AudioBus = InGeneratorParams.AudioBus)
 			{
 				BusId = AudioBus->GetUniqueID();
@@ -195,6 +198,9 @@ namespace AudioModulation
 
 		}
 
+		// Retrieving/mutating the MixerDevice is only safe during Update calls 
+		// on the audio render thread if it is not called during FAudioDevice::Teardown, 
+		// so it is reset prior to that via the OnDeviceDestroyed callback
 		void CreatePatchForBus()
 		{
 			if (BusId != INDEX_NONE)
@@ -213,6 +219,17 @@ namespace AudioModulation
 			}
 		}
 
+		void OnDeviceDestroyed(Audio::FDeviceId InDeviceId)
+		{
+			if (InDeviceId == AudioDeviceId)
+			{
+				// Reset bus id on device destruction to avoid reinitializing
+				// it during FAudioDevice::Teardown via Update.
+				BusId = INDEX_NONE;
+				FAudioDeviceManagerDelegates::OnAudioDeviceDestroyed.Remove(DeviceDestroyedHandle);
+			}
+		}
+
 	private:
 		Audio::FPatchOutputStrongPtr AudioBusPatch;
 		Audio::FAlignedFloatBuffer TempBuffer;
@@ -226,6 +243,8 @@ namespace AudioModulation
 		uint8 bBypass : 1;
 		uint8 bInvert : 1;
 		uint8 bBusRequiresPatch : 1;
+
+		FDelegateHandle DeviceDestroyedHandle;
 	};
 } // namespace AudioModulation
 
