#UE Disable expensive audio modulation asset registry OnAssetAdded function

diff --git a/Engine/Plugins/Runtime/AudioModulation/Source/AudioModulationEditor/Private/AudioModulationEditor.cpp b/Engine/Plugins/Runtime/AudioModulation/Source/AudioModulationEditor/Private/AudioModulationEditor.cpp
--- a/Engine/Plugins/Runtime/AudioModulation/Source/AudioModulationEditor/Private/AudioModulationEditor.cpp
+++ b/Engine/Plugins/Runtime/AudioModulation/Source/AudioModulationEditor/Private/AudioModulationEditor.cpp
@@ -133,6 +133,7 @@ void FAudioModulationEditorModule::StartupModule()
 	// All parameters are required to always be loaded in editor to enable them to be referenced via object
 	// metadata and custom layouts, even if they are not referenced by runtime uobjects/systems directly
 	IAssetRegistry& AssetRegistry = FModuleManager::LoadModuleChecked<FAssetRegistryModule>("AssetRegistry").Get();
+	/*
 	AssetRegistry.OnAssetAdded().AddLambda([](const FAssetData& InAssetData)
 	{
 		UClass* AssetDataClass = InAssetData.GetClass();
@@ -144,6 +145,7 @@ void FAudioModulationEditorModule::StartupModule()
 			}
 		}
 	});
+	*/
 	AssetRegistry.OnInMemoryAssetDeleted().AddLambda([](UObject* ObjectDeleted)
 	{
 		if (USoundModulationParameter* Parameter = Cast<USoundModulationParameter>(ObjectDeleted))
