Deactivate All Bus Mixes BP Function

diff --git a/Engine/Plugins/Runtime/AudioModulation/Source/AudioModulation/Private/AudioModulationStatics.cpp b/Engine/Plugins/Runtime/AudioModulation/Source/AudioModulation/Private/AudioModulationStatics.cpp
--- a/Engine/Plugins/Runtime/AudioModulation/Source/AudioModulation/Private/AudioModulationStatics.cpp
+++ b/Engine/Plugins/Runtime/AudioModulation/Source/AudioModulation/Private/AudioModulationStatics.cpp
@@ -423,6 +423,15 @@ void UAudioModulationStatics::ClearAllGlobalBusMixValues(const UObject* WorldCon
 	}
 }
 
+void UAudioModulationStatics::DeactivateAllBusMixes(const UObject* WorldContextObject)
+{
+	UWorld* World = GetAudioWorld(WorldContextObject);
+	if (AudioModulation::FAudioModulationManager* ModSystem = GetModulation(World))
+	{
+		ModSystem->DeactivateAllBusMixes();
+	}
+}
+
 void UAudioModulationStatics::UpdateMixByFilter(
 	const UObject* WorldContextObject,
 	USoundControlBusMix* Mix,
diff --git a/Engine/Plugins/Runtime/AudioModulation/Source/AudioModulation/Public/AudioModulationStatics.h b/Engine/Plugins/Runtime/AudioModulation/Source/AudioModulation/Public/AudioModulationStatics.h
--- a/Engine/Plugins/Runtime/AudioModulation/Source/AudioModulation/Public/AudioModulationStatics.h
+++ b/Engine/Plugins/Runtime/AudioModulation/Source/AudioModulation/Public/AudioModulationStatics.h
@@ -266,6 +266,14 @@ public:
 	)
 	static void ClearAllGlobalBusMixValues(const UObject* WorldContextObject, float FadeTime = -1.0f);
 
+	/** Deactivates all currently active Control Bus Mixes. This includes the Global Control Bus Mixes.
+	 */
+	UFUNCTION(BlueprintCallable, Category = "Audio|Modulation", DisplayName = "Deactivate All Control Bus Mixes", meta = (
+		WorldContext = "WorldContextObject",
+		Keywords = "modulation modulator stage")
+	)
+	static void DeactivateAllBusMixes(const UObject* WorldContextObject);
+
 	/** Sets filtered stages of a given class to a provided target value for active instance of mix.
 	 * Does not update UObject definition of mix.
 	 * @param Mix - Mix to modify
