Audio Modulation AD Envelope Generator

diff --git a/Engine/Plugins/Runtime/AudioModulation/Source/AudioModulation/Private/AudioModulationStatics.cpp b/Engine/Plugins/Runtime/AudioModulation/Source/AudioModulation/Private/AudioModulationStatics.cpp
--- a/Engine/Plugins/Runtime/AudioModulation/Source/AudioModulation/Private/AudioModulationStatics.cpp
+++ b/Engine/Plugins/Runtime/AudioModulation/Source/AudioModulation/Private/AudioModulationStatics.cpp
@@ -310,6 +310,24 @@ USoundModulationGeneratorLFO* UAudioModulationStatics::CreateLFOGenerator(UObjec
 	return NewGenerator;
 }
 
+USoundModulationGeneratorADEnvelope* UAudioModulationStatics::CreateADEnvelopeGenerator(UObject* WorldContextObject, FName Name, const FSoundModulationADEnvelopeParams& Params)
+{
+	UWorld* World = GetAudioWorld(WorldContextObject);
+	if (!World)
+	{
+		return nullptr;
+	}
+
+	USoundModulationGeneratorADEnvelope* NewGenerator = NewObject<USoundModulationGeneratorADEnvelope>(WorldContextObject, USoundModulationGeneratorADEnvelope::StaticClass(), Name);
+
+	if (NewGenerator)
+	{
+		NewGenerator->Params = Params;
+	}
+
+	return NewGenerator;
+}
+
 void UAudioModulationStatics::DeactivateBus(const UObject* WorldContextObject, USoundControlBus* Bus)
 {
 	if (Bus)
diff --git a/Engine/Plugins/Runtime/AudioModulation/Source/AudioModulation/Private/Generators/SoundModulationADEnvelope.cpp b/Engine/Plugins/Runtime/AudioModulation/Source/AudioModulation/Private/Generators/SoundModulationADEnvelope.cpp
new file mode 100644
--- /dev/null
+++ b/Engine/Plugins/Runtime/AudioModulation/Source/AudioModulation/Private/Generators/SoundModulationADEnvelope.cpp
@@ -0,0 +1,135 @@
+// Copyright Epic Games, Inc. All Rights Reserved.
+#include "Generators/SoundModulationADEnvelope.h"
+
+#include "AudioDevice.h"
+#include "AudioModulation.h"
+#include "AudioModulationLogging.h"
+#include "AudioModulationSystem.h"
+#include "Engine/World.h"
+
+#include UE_INLINE_GENERATED_CPP_BY_NAME(SoundModulationADEnvelope)
+
+
+namespace AudioModulation
+{
+	namespace GeneratorADEnvelopePrivate
+	{
+		static const TArray<FString> DebugCategories =
+		{
+			TEXT("Value"),
+			TEXT("Attack Time"),
+			TEXT("Decay Time"),
+			TEXT("Attack Curve"),
+			TEXT("Decay Curve"),
+			TEXT("Looping")
+		};
+
+		static const FString DebugName = TEXT("AD Envelope");
+	} // GeneratorADEnvelopePrivate
+
+	class FADEnvelopeGenerator : public IGenerator
+	{
+		public:
+			FADEnvelopeGenerator() = default;
+
+			FADEnvelopeGenerator(const FADEnvelopeGenerator& InGenerator)
+				: Envelope(InGenerator.Envelope)
+				, Value(InGenerator.Value)
+				, Params(InGenerator.Params)
+			{
+			}
+
+			FADEnvelopeGenerator(const FSoundModulationADEnvelopeParams& InParams)
+				: Params(InParams)
+			{
+				Envelope.SetAttackTimeSeconds(InParams.AttackTime);
+				Envelope.SetDecayTimeSeconds(InParams.DecayTime);
+				Envelope.SetAttackCurveFactor(InParams.AttackCurve);
+				Envelope.SetDecayCurveFactor(InParams.DecayCurve);
+				Envelope.SetLooping(InParams.bLooping);
+				Envelope.Attack();
+			}
+
+			virtual ~FADEnvelopeGenerator() = default;
+
+			virtual float GetValue() const override
+			{
+				return Value;
+			}
+
+			virtual bool IsBypassed() const override
+			{
+				return Params.bBypass;
+			}
+
+			virtual void Update(double InElapsed) override
+			{
+				if (InElapsed > 0.0f)
+				{
+					const float SampleRate = static_cast<float>(1.0 / InElapsed);
+					Envelope.Init(SampleRate);
+					Envelope.GetNextEnvelopeOut(Value);
+				}
+			}
+
+			virtual FGeneratorPtr Clone() const override
+			{
+				return FGeneratorPtr(new FADEnvelopeGenerator(*this));
+			}
+
+			virtual void UpdateGenerator(FGeneratorPtr&& InGenerator) override
+			{
+				if (!ensure(InGenerator.IsValid()))
+				{
+					return;
+				}
+
+				AudioRenderThreadCommand([this, NewGenerator = MoveTemp(InGenerator)]()
+				{
+					const FADEnvelopeGenerator* Generator = static_cast<const FADEnvelopeGenerator*>(NewGenerator.Get());
+					Params = Generator->Params;
+
+					Envelope.SetAttackTimeSeconds(Params.AttackTime);
+					Envelope.SetDecayTimeSeconds(Params.DecayTime);
+					Envelope.SetAttackCurveFactor(Params.AttackCurve);
+					Envelope.SetDecayCurveFactor(Params.DecayCurve);
+					Envelope.SetLooping(Params.bLooping);
+				});
+			}
+
+#if !UE_BUILD_SHIPPING
+			virtual void GetDebugCategories(TArray<FString>& OutDebugCategories) const override
+			{
+				OutDebugCategories = GeneratorADEnvelopePrivate::DebugCategories;
+			}
+
+			virtual const FString& GetDebugName() const override
+			{
+				return GeneratorADEnvelopePrivate::DebugName;
+			}
+
+			virtual void GetDebugValues(TArray<FString>& OutDebugValues) const override
+			{
+				OutDebugValues.Add(FString::Printf(TEXT("%.4f"), GetValue()));
+				OutDebugValues.Add(FString::Printf(TEXT("%.4f"), Params.AttackTime));
+				OutDebugValues.Add(FString::Printf(TEXT("%.4f"), Params.DecayTime));
+				OutDebugValues.Add(FString::Printf(TEXT("%.4f"), Params.AttackCurve));
+				OutDebugValues.Add(FString::Printf(TEXT("%.4f"), Params.DecayCurve));
+				OutDebugValues.Add(FString::Printf(TEXT("%i"), Envelope.IsLooping()));
+			}
+#endif // !UE_BUILD_SHIPPING
+
+	protected:
+		Audio::FADEnvelope Envelope;
+		float Value = 1.0f;
+		FSoundModulationADEnvelopeParams Params;
+	};
+} // namespace AudioModulation
+
+AudioModulation::FGeneratorPtr USoundModulationGeneratorADEnvelope::CreateInstance() const
+{
+	using namespace AudioModulation;
+
+	return FGeneratorPtr(new FADEnvelopeGenerator(Params));
+}
+
diff --git a/Engine/Plugins/Runtime/AudioModulation/Source/AudioModulation/Public/AudioModulationStatics.h b/Engine/Plugins/Runtime/AudioModulation/Source/AudioModulation/Public/AudioModulationStatics.h
--- a/Engine/Plugins/Runtime/AudioModulation/Source/AudioModulation/Public/AudioModulationStatics.h
+++ b/Engine/Plugins/Runtime/AudioModulation/Source/AudioModulation/Public/AudioModulationStatics.h
@@ -1,6 +1,7 @@
 // Copyright Epic Games, Inc. All Rights Reserved.
 #pragma once
 
+#include "Generators/SoundModulationADEnvelope.h"
 #include "Generators/SoundModulationEnvelopeFollower.h"
 #include "Generators/SoundModulationLFO.h"
 #include "Kismet/BlueprintFunctionLibrary.h"
@@ -163,6 +164,20 @@ public:
 			FName Name,
 			FSoundModulationLFOParams Params);
 
+	/** Creates a modulation generator based on an Attack/Decay Envelope.
+	 * @param Name - Name of generator.
+	 * @param Settings - The AD Envelope Settings.
+	 * @return Capture this in a Blueprint variable to prevent it from being garbage collected.
+	 */
+	UFUNCTION(BlueprintCallable, Category = "Audio|Modulation", DisplayName = "Create AD Envelope Generator", meta = (
+		WorldContext = "WorldContextObject",
+		Keywords = "make modulator")
+	)
+		static UPARAM(DisplayName = "Generator") USoundModulationGeneratorADEnvelope* CreateADEnvelopeGenerator(
+			UObject* WorldContextObject,
+			FName Name,
+			const FSoundModulationADEnvelopeParams& Params);
+
 	/** Deactivates a bus. Does nothing if the provided bus is already inactive.
 	 * @param Bus - Scope of modulator
 	 */
diff --git a/Engine/Plugins/Runtime/AudioModulation/Source/AudioModulation/Public/Generators/SoundModulationADEnvelope.h b/Engine/Plugins/Runtime/AudioModulation/Source/AudioModulation/Public/Generators/SoundModulationADEnvelope.h
new file mode 100644
--- /dev/null
+++ b/Engine/Plugins/Runtime/AudioModulation/Source/AudioModulation/Public/Generators/SoundModulationADEnvelope.h
@@ -0,0 +1,54 @@
+// Copyright Epic Games, Inc. All Rights Reserved.
+#pragma once
+
+#include "DSP/Envelope.h"
+#include "IAudioModulation.h"
+#include "SoundModulationGenerator.h"
+#include "Templates/SharedPointer.h"
+#include "UObject/ObjectMacros.h"
+#include "UObject/Object.h"
+
+#include "SoundModulationADEnvelope.generated.h"
+
+USTRUCT(BlueprintType)
+struct FSoundModulationADEnvelopeParams
+{
+	GENERATED_USTRUCT_BODY()
+
+	/** Attack time of the envelope (seconds). */
+	UPROPERTY(EditAnywhere, Category = Modulation, BlueprintReadWrite, meta = (DisplayPriority = 20, EditCondition = "!bBypass", UIMin = "0", ClampMin = "0"))
+	float AttackTime = 0.0f;
+
+	/** Decay time of the envelope (seconds). */
+	UPROPERTY(EditAnywhere, Category = Modulation, BlueprintReadWrite, meta = (DisplayPriority = 30, EditCondition = "!bBypass", UIMin = "0", ClampMin = "0"))
+	float DecayTime = 1.0f;
+
+	/** The exponential curve factor of the attack. 1.0 = linear growth, < 1.0 logorithmic growth, > 1.0 exponential growth. */
+	UPROPERTY(EditAnywhere, Category = Modulation, BlueprintReadWrite, meta = (DisplayPriority = 40, EditCondition = "!bBypass", UIMin = "0", ClampMin = "0"))
+	float AttackCurve = 1.0f;
+
+	/** The exponential curve factor of the decay. 1.0 = linear decay, < 1.0 exponential decay, > 1.0 logarithmic decay. */
+	UPROPERTY(EditAnywhere, Category = Modulation, BlueprintReadWrite, meta = (DisplayPriority = 50, EditCondition = "!bBypass", UIMin = "0", ClampMin = "0"))
+	float DecayCurve = 1.0f;
+
+	/** Whether or not to loop the envelope. */
+	UPROPERTY(EditAnywhere, Category = Modulation, BlueprintReadWrite, meta = (DisplayPriority = 60, EditCondition = "!bBypass"))
+	bool bLooping = true;
+
+	/** If true, bypasses envelope bus from being modulated by parameters, patches, or mixed (Envelope remains active and computed). */
+	UPROPERTY(EditAnywhere, Category = Modulation, BlueprintReadWrite, meta = (DisplayPriority = 10))
+	bool bBypass = false;
+};
+
+
+UCLASS(BlueprintType, hidecategories = Object, editinlinenew, MinimalAPI)
+class USoundModulationGeneratorADEnvelope : public USoundModulationGenerator
+{
+	GENERATED_BODY()
+
+public:
+	UPROPERTY(EditAnywhere, Category = Modulation, BlueprintReadWrite, meta = (ShowOnlyInnerProperties))
+	FSoundModulationADEnvelopeParams Params;
+
+	virtual AudioModulation::FGeneratorPtr CreateInstance() const override;
+};
