add tooltip to Control Bus Mix Stage "Value" parameter

diff --git a/Engine/Plugins/Runtime/AudioModulation/Source/AudioModulation/Public/SoundControlBusMix.h b/Engine/Plugins/Runtime/AudioModulation/Source/AudioModulation/Public/SoundControlBusMix.h
--- a/Engine/Plugins/Runtime/AudioModulation/Source/AudioModulation/Public/SoundControlBusMix.h
+++ b/Engine/Plugins/Runtime/AudioModulation/Source/AudioModulation/Public/SoundControlBusMix.h
@@ -22,7 +22,7 @@ struct AUDIOMODULATION_API FSoundControlBusMixStage
 	FSoundControlBusMixStage();
 	FSoundControlBusMixStage(USoundControlBus* InBus, const float TargetValue);
 
-	/* Bus controlled by stage. */
+	/* Bus controlled by the mix stage. */
 	UPROPERTY(EditAnywhere, Category = Stage, BlueprintReadWrite)
 	TObjectPtr<USoundControlBus> Bus;
 
diff --git a/Engine/Plugins/Runtime/AudioModulation/Source/AudioModulationEditor/Private/Layouts/SoundControlBusMixStageLayout.cpp b/Engine/Plugins/Runtime/AudioModulation/Source/AudioModulationEditor/Private/Layouts/SoundControlBusMixStageLayout.cpp
--- a/Engine/Plugins/Runtime/AudioModulation/Source/AudioModulationEditor/Private/Layouts/SoundControlBusMixStageLayout.cpp
+++ b/Engine/Plugins/Runtime/AudioModulation/Source/AudioModulationEditor/Private/Layouts/SoundControlBusMixStageLayout.cpp
@@ -223,7 +223,7 @@ void FSoundControlBusMixStageLayoutCustomization::CustomizeChildren(TSharedRef<I
 		SNew(STextBlock)
 		.Font(IDetailLayoutBuilder::GetDetailFont())
 		.Text(LOCTEXT("BusModulationValue_MixValueName", "Value"))
-		.ToolTipText(StructPropertyHandle->GetToolTipText())
+		.ToolTipText(LOCTEXT("BusModulationValue_MixValueTooltip", "Value the mix stage is set to. If the bus's parameter has a unit conversion (eg. dB), there will be a unit value followed by a normalized (0-1) value. Changing one will update the other to have the corresponding value."))
 	]
 	.ValueContent()
 		.MinDesiredWidth(300.0f)
