Update Modulation Generator to use CreateProxyData instead of CreateNewProxyData (deprecated)

diff --git a/Engine/Plugins/Runtime/AudioModulation/Source/AudioModulation/Private/SoundModulationGenerator.cpp b/Engine/Plugins/Runtime/AudioModulation/Source/AudioModulation/Private/SoundModulationGenerator.cpp
--- a/Engine/Plugins/Runtime/AudioModulation/Source/AudioModulation/Private/SoundModulationGenerator.cpp
+++ b/Engine/Plugins/Runtime/AudioModulation/Source/AudioModulation/Private/SoundModulationGenerator.cpp
@@ -46,10 +46,10 @@ void USoundModulationGenerator::PostEditChangeProperty(FPropertyChangedEvent& In
 }
 #endif // WITH_EDITOR
 
-TUniquePtr<Audio::IProxyData> USoundModulationGenerator::CreateNewProxyData(const Audio::FProxyDataInitParams& InitParams)
+TSharedPtr<Audio::IProxyData> USoundModulationGenerator::CreateProxyData(const Audio::FProxyDataInitParams& InitParams)
 {
 	using namespace AudioModulation;
-	return MakeUnique<FSoundModulatorAssetProxy>(*this);
+	return MakeShared<FSoundModulatorAssetProxy>(*this);
 }
 
 TUniquePtr<Audio::IModulatorSettings> USoundModulationGenerator::CreateProxySettings() const
diff --git a/Engine/Plugins/Runtime/AudioModulation/Source/AudioModulation/Public/SoundModulationGenerator.h b/Engine/Plugins/Runtime/AudioModulation/Source/AudioModulation/Public/SoundModulationGenerator.h
--- a/Engine/Plugins/Runtime/AudioModulation/Source/AudioModulation/Public/SoundModulationGenerator.h
+++ b/Engine/Plugins/Runtime/AudioModulation/Source/AudioModulation/Public/SoundModulationGenerator.h
@@ -86,7 +86,7 @@ public:
 	}
 
 	/* USoundModulatorBase Implementation */
-	virtual TUniquePtr<Audio::IProxyData> CreateNewProxyData(const Audio::FProxyDataInitParams& InitParams) override;
+	virtual TSharedPtr<Audio::IProxyData> CreateProxyData(const Audio::FProxyDataInitParams& InitParams) override;
 	virtual TUniquePtr<Audio::IModulatorSettings> CreateProxySettings() const override;
 
 #if WITH_EDITOR
