Add warnings in cases where CBM Mix Profiles have mismatching Mix Stages

diff --git a/Engine/Plugins/Runtime/AudioModulation/Source/AudioModulation/Private/AudioModulationSystem.cpp b/Engine/Plugins/Runtime/AudioModulation/Source/AudioModulation/Private/AudioModulationSystem.cpp
--- a/Engine/Plugins/Runtime/AudioModulation/Source/AudioModulation/Private/AudioModulationSystem.cpp
+++ b/Engine/Plugins/Runtime/AudioModulation/Source/AudioModulation/Private/AudioModulationSystem.cpp
@@ -759,11 +759,16 @@ namespace AudioModulation
 			}
 		}
 	
-		RunCommandOnProcessingThread([this, MixId, StageSettings, InFadeTime]()
+		const FString BusMixName = InOutMix.GetName();
+		RunCommandOnProcessingThread([this, MixId, StageSettings, InFadeTime, BusMixName]()
 		{
 			if (FModulatorBusMixProxy* BusMixes = RefProxies.BusMixes.Find(MixId))
 			{
-				BusMixes->SetMix(StageSettings, InFadeTime);
+				BusMixes->SetMix(StageSettings, InFadeTime, BusMixName);
+			}
+			else
+			{
+				UE_LOG(LogAudioModulation, Warning, TEXT("Updating Control Bus Mix failed: Bus Mix '%s' not currently active. Buses without any valid Mix Stages will not activate."), *BusMixName);
 			}
 		});
 	}
diff --git a/Engine/Plugins/Runtime/AudioModulation/Source/AudioModulation/Private/SoundControlBusMixProxy.cpp b/Engine/Plugins/Runtime/AudioModulation/Source/AudioModulation/Private/SoundControlBusMixProxy.cpp
--- a/Engine/Plugins/Runtime/AudioModulation/Source/AudioModulation/Private/SoundControlBusMixProxy.cpp
+++ b/Engine/Plugins/Runtime/AudioModulation/Source/AudioModulation/Private/SoundControlBusMixProxy.cpp
@@ -119,7 +119,7 @@ namespace AudioModulation
 		}
 	}
 
-	void FModulatorBusMixProxy::SetMix(const TArray<FModulatorBusMixStageSettings>& InStages, float InFadeTime)
+	void FModulatorBusMixProxy::SetMix(const TArray<FModulatorBusMixStageSettings>& InStages, float InFadeTime, const FString& BusMixName)
 	{
 		for (const FModulatorBusMixStageSettings& NewStage : InStages)
 		{
@@ -133,6 +133,10 @@ namespace AudioModulation
 				// Setting entire mix wipes pre-existing user fade requests
 				StageProxy->Value.SetActiveFade(FSoundModulationMixValue::EActiveFade::Override, InFadeTime);
 			}
+			else
+			{
+				UE_LOG(LogAudioModulation, Warning, TEXT("Bus '%s' Not currently applied to Bus Mix '%s'. Please ensure that all your Mix Profiles have the same Control Buses."), *NewStage.Address, *BusMixName);
+			}
 		}
 	}
 
diff --git a/Engine/Plugins/Runtime/AudioModulation/Source/AudioModulation/Private/SoundControlBusMixProxy.h b/Engine/Plugins/Runtime/AudioModulation/Source/AudioModulation/Private/SoundControlBusMixProxy.h
--- a/Engine/Plugins/Runtime/AudioModulation/Source/AudioModulation/Private/SoundControlBusMixProxy.h
+++ b/Engine/Plugins/Runtime/AudioModulation/Source/AudioModulation/Private/SoundControlBusMixProxy.h
@@ -75,7 +75,7 @@ namespace AudioModulation
 		void Reset();
 
 		void SetEnabled(FModulatorBusMixSettings&& InSettings);
-		void SetMix(const TArray<FModulatorBusMixStageSettings>& InStages, float InFadeTime);
+		void SetMix(const TArray<FModulatorBusMixStageSettings>& InStages, float InFadeTime, const FString& BusMixName);
 		void SetMixByFilter(const FString& InAddressFilter, uint32 InParamClassId, uint32 InParamId, float InValue, float InFadeTime);
 		void SetStopping();
 
