Update Modulation Parameter comments to explain the purpose of the Default Parameter Value field

diff --git a/Engine/Plugins/Runtime/AudioModulation/Source/AudioModulation/Public/SoundModulationParameter.h b/Engine/Plugins/Runtime/AudioModulation/Source/AudioModulation/Public/SoundModulationParameter.h
--- a/Engine/Plugins/Runtime/AudioModulation/Source/AudioModulation/Public/SoundModulationParameter.h
+++ b/Engine/Plugins/Runtime/AudioModulation/Source/AudioModulation/Public/SoundModulationParameter.h
@@ -16,13 +16,7 @@ struct AUDIOMODULATION_API FSoundModulationParameterSettings
 {
 	GENERATED_USTRUCT_BODY()
 
-	/** 
-	  * Normalized, unitless default value of modulator. To ensure bypass functionality of mixing, patching, and modulating 
-	  * functions as anticipated, value should be selected such that GetMixFunction (see USoundModulationParameter)
-	  * reduces to an identity function (i.e. function acts as a "pass-through" for all values in the range [0.0, 1.0]).
-	  * If GetMixFunction performs the mathematical operation f(x1, x2), then the default ValueNormalized should result in
-	  * f(x1, d) = x1 where d is ValueNormalized.
-	  */
+	/** Default value of modulator (unitless) */
 	UPROPERTY(EditAnywhere, Category = General, BlueprintReadOnly, meta = (ClampMin = "0.0", ClampMax = "1.0"))
 	float ValueNormalized = 1.0f;
 
@@ -31,7 +25,12 @@ struct AUDIOMODULATION_API FSoundModulationParameterSettings
 	UPROPERTY(EditAnywhere, Category = General)
 	FText UnitDisplayName;
 
-	/** Default value of modulator in units (editor only) */
+	/**
+	  * Default value of the modulator. To ensure bypass functionality of mixing, patching, and modulating
+	  * functions as anticipated, value should be selected such that GetMixFunction (see USoundModulationParameter)
+	  * reduces to an identity function (i.e. function acts as a "pass-through" for all values in the range [0.0, 1.0]).
+	  * That is to say, this should be set to the value which has no effect on the sound.
+	  */
 	UPROPERTY(Transient, EditAnywhere, Category = General)
 	float ValueUnit = 1.0f;
 #endif // WITH_EDITORONLY_DATA
