Create Default Volume Modulation Parameter when the asset is not found

diff --git a/Engine/Plugins/Runtime/AudioModulation/Source/AudioModulation/Private/AudioModulationSettings.cpp b/Engine/Plugins/Runtime/AudioModulation/Source/AudioModulation/Private/AudioModulationSettings.cpp
--- a/Engine/Plugins/Runtime/AudioModulation/Source/AudioModulation/Private/AudioModulationSettings.cpp
+++ b/Engine/Plugins/Runtime/AudioModulation/Source/AudioModulation/Private/AudioModulationSettings.cpp
@@ -37,7 +37,40 @@ void UAudioModulationSettings::RegisterParameters() const
 		}
 		else
 		{
-			UE_LOG(LogAudioModulation, Error, TEXT("Failed to load parameter at '%s': Missing asset or invalid type."), *Path.GetAssetName());
+			const FString& ParamName = Path.GetAssetName();
+			UE_LOG(LogAudioModulation, Error, TEXT("Failed to load parameter at '%s': Missing asset or invalid type."), *ParamName);
+
+			// Create a parameter procedurally instead. These parameters are all required by Audio Mixer Sources.
+			if (ParamName == "HPFCutoffFrequency")
+			{
+				Audio::FModulationParameter NewParam = USoundModulationParameterHPFFrequency::CreateDefaultParameter();
+				NewParam.ParameterName = "HPFCutoffFrequency";
+				Audio::RegisterModulationParameter(NewParam.ParameterName, MoveTemp(NewParam));
+				UE_LOG(LogAudioModulation, Error, TEXT("Created temporary HPFCutoffFrequency parameter to use instead."));
+			}
+			else if (ParamName == "LPFCutoffFrequency")
+			{
+				Audio::FModulationParameter NewParam = USoundModulationParameterLPFFrequency::CreateDefaultParameter();
+				NewParam.ParameterName = "LPFCutoffFrequency";
+				Audio::RegisterModulationParameter(NewParam.ParameterName, MoveTemp(NewParam));
+				UE_LOG(LogAudioModulation, Error, TEXT("Created temporary LPFCutoffFrequency parameter to use instead."));
+			}
+			else if (ParamName == "Pitch")
+			{
+				// The Pitch parameter is a specialized Bipolar parameter, so we have to give it extra data and change the unit display name
+				Audio::FModulationParameter NewParam = USoundModulationParameterBipolar::CreateDefaultParameter(24.0f);
+				NewParam.ParameterName = "Pitch";
+				Audio::RegisterModulationParameter(NewParam.ParameterName, MoveTemp(NewParam));
+				UE_LOG(LogAudioModulation, Error, TEXT("Created temporary Pitch parameter to use instead."));
+			}
+			else if (ParamName == "Volume")
+			{
+				// TODO: Similar to pitch, send -60 dB to the function
+				Audio::FModulationParameter NewParam = USoundModulationParameterVolume::CreateDefaultParameter(-60.0f);
+				NewParam.ParameterName = "Volume";
+				Audio::RegisterModulationParameter(NewParam.ParameterName, MoveTemp(NewParam));
+				UE_LOG(LogAudioModulation, Error, TEXT("Created temporary Volume parameter to use instead."));
+			}
 		}
 	}
 }
diff --git a/Engine/Plugins/Runtime/AudioModulation/Source/AudioModulation/Private/SoundModulationParameter.cpp b/Engine/Plugins/Runtime/AudioModulation/Source/AudioModulation/Private/SoundModulationParameter.cpp
--- a/Engine/Plugins/Runtime/AudioModulation/Source/AudioModulation/Private/SoundModulationParameter.cpp
+++ b/Engine/Plugins/Runtime/AudioModulation/Source/AudioModulation/Private/SoundModulationParameter.cpp
@@ -105,6 +105,15 @@ Audio::FModulationMixFunction USoundModulationParameterHPFFrequency::GetMixFunct
 	};
 }
 
+Audio::FModulationParameter USoundModulationParameterHPFFrequency::CreateDefaultParameter()
+{
+	UClass* ThisClass = USoundModulationParameterHPFFrequency::StaticClass();
+	check(ThisClass);
+	USoundModulationParameterHPFFrequency* ClassDefault = ThisClass->GetDefaultObject<USoundModulationParameterHPFFrequency>();
+	check(ClassDefault);
+	return ClassDefault->CreateParameter();
+}
+
 USoundModulationParameterHPFFrequency::USoundModulationParameterHPFFrequency(const FObjectInitializer& ObjectInitializer)
 	: Super(ObjectInitializer)
 {
@@ -123,6 +132,15 @@ Audio::FModulationMixFunction USoundModulationParameterLPFFrequency::GetMixFunct
 	};
 }
 
+Audio::FModulationParameter USoundModulationParameterLPFFrequency::CreateDefaultParameter()
+{
+	UClass* ThisClass = USoundModulationParameterLPFFrequency::StaticClass();
+	check(ThisClass);
+	USoundModulationParameterLPFFrequency* ClassDefault = ThisClass->GetDefaultObject<USoundModulationParameterLPFFrequency>();
+	check(ClassDefault);
+	return ClassDefault->CreateParameter();
+}
+
 bool USoundModulationParameterScaled::RequiresUnitConversion() const
 {
 	return true;
@@ -204,6 +222,25 @@ float USoundModulationParameterBipolar::GetUnitMin() const
 	return UnitRange * -0.5f;
 }
 
+Audio::FModulationParameter USoundModulationParameterBipolar::CreateDefaultParameter(float UnitRange)
+{
+	UClass* ThisClass = USoundModulationParameterBipolar::StaticClass();
+	check(ThisClass);
+	USoundModulationParameterBipolar* ClassDefault = ThisClass->GetDefaultObject<USoundModulationParameterBipolar>();
+	check(ClassDefault);
+	
+	// Cache the default value for Bipolar Range because we need to change it
+	const float DefaultRange = ClassDefault->UnitRange;
+	ClassDefault->UnitRange = UnitRange;
+	
+	Audio::FModulationParameter DefaultParam = ClassDefault->CreateParameter();
+	
+	// Set the default back to what it was before
+	ClassDefault->UnitRange = DefaultRange;
+
+	return DefaultParam;
+}
+
 bool USoundModulationParameterVolume::RequiresUnitConversion() const
 {
 	return true;
@@ -239,6 +276,25 @@ float USoundModulationParameterVolume::GetUnitMax() const
 	return 0.0f;
 }
 
+Audio::FModulationParameter USoundModulationParameterVolume::CreateDefaultParameter(float MinUnitVolume)
+{
+	UClass* ThisClass = USoundModulationParameterVolume::StaticClass();
+	check(ThisClass);
+	USoundModulationParameterVolume* ClassDefault = ThisClass->GetDefaultObject<USoundModulationParameterVolume>();
+	check(ClassDefault);
+
+	// Cache the default value for Bipolar Range because we need to change it
+	const float DefaultMinVolume = ClassDefault->MinVolume;
+	ClassDefault->MinVolume = MinUnitVolume;
+
+	Audio::FModulationParameter NewParam = ClassDefault->CreateParameter();
+
+	// Set the default back to what it was before
+	ClassDefault->MinVolume = DefaultMinVolume;
+
+	return NewParam;
+}
+
 namespace AudioModulation
 {
 	const Audio::FModulationParameter& GetOrRegisterParameter(const USoundModulationParameter* InParameter, const FString& InBreadcrumb)
diff --git a/Engine/Plugins/Runtime/AudioModulation/Source/AudioModulation/Public/SoundModulationParameter.h b/Engine/Plugins/Runtime/AudioModulation/Source/AudioModulation/Public/SoundModulationParameter.h
--- a/Engine/Plugins/Runtime/AudioModulation/Source/AudioModulation/Public/SoundModulationParameter.h
+++ b/Engine/Plugins/Runtime/AudioModulation/Source/AudioModulation/Public/SoundModulationParameter.h
@@ -203,6 +203,7 @@ class USoundModulationParameterLPFFrequency : public USoundModulationParameterFi
 
 public:
 	virtual Audio::FModulationMixFunction GetMixFunction() const override;
+	static Audio::FModulationParameter CreateDefaultParameter();
 };
 
 // Modulation Parameter that scales normalized, unitless value to logarithmic frequency unit space with standard filter min and max frequency set.
@@ -214,6 +215,7 @@ class USoundModulationParameterHPFFrequency : public USoundModulationParameterFi
 
 public:
 	virtual Audio::FModulationMixFunction GetMixFunction() const override;
+	static Audio::FModulationParameter CreateDefaultParameter();
 };
 
 // Modulation Parameter that scales normalized, unitless value to bipolar range. Mixes multiplicatively.
@@ -233,6 +235,7 @@ public:
 	virtual Audio::FModulationNormalizedConversionFunction GetNormalizedConversionFunction() const override;
 	virtual float GetUnitMax() const override;
 	virtual float GetUnitMin() const override;
+	static Audio::FModulationParameter CreateDefaultParameter(float UnitRange = 2.0f);
 };
 
 UCLASS(BlueprintType, MinimalAPI)
@@ -250,6 +253,7 @@ public:
 	virtual Audio::FModulationNormalizedConversionFunction GetNormalizedConversionFunction() const override;
 	virtual float GetUnitMin() const override;
 	virtual float GetUnitMax() const override;
+	static Audio::FModulationParameter CreateDefaultParameter(float MinUnitVolume = -100.0f);
 };
 
 namespace AudioModulation
