Remove impossible cast to USoundControlBusMix from USoundModulatorBase
#rnx

diff --git a/Engine/Plugins/Runtime/AudioModulation/Source/AudioModulation/Private/AudioModulationSystem.cpp b/Engine/Plugins/Runtime/AudioModulation/Source/AudioModulation/Private/AudioModulationSystem.cpp
--- a/Engine/Plugins/Runtime/AudioModulation/Source/AudioModulation/Private/AudioModulationSystem.cpp
+++ b/Engine/Plugins/Runtime/AudioModulation/Source/AudioModulation/Private/AudioModulationSystem.cpp
@@ -903,25 +903,6 @@ namespace AudioModulation
 			});
 		}
 
-		if (const USoundControlBusMix* InMix = Cast<USoundControlBusMix>(&InModulator))
-		{
-			RunCommandOnProcessingThread([this, BusMixSettings = FModulatorBusMixSettings(*InMix)]() mutable
-			{
-				FBusMixHandle BusMixHandle = FBusMixHandle::Get(BusMixSettings.GetId(), RefProxies.BusMixes);
-				if (BusMixHandle.IsValid())
-				{
-					FModulatorBusMixProxy& BusMixProxy = BusMixHandle.FindProxy();
-					BusMixProxy = MoveTemp(BusMixSettings);
-				}
-#if !UE_BUILD_SHIPPING
-				else
-				{
-					UE_LOG(LogAudioModulation, Verbose, TEXT("Update to '%s' Ignored: Control Bus Mix is inactive."), *BusMixSettings.GetName());
-				}
-#endif // !UE_BUILD_SHIPPING
-			});
-		}
-
 		if (const USoundModulationPatch* InPatch = Cast<USoundModulationPatch>(&InModulator))
 		{
 			RunCommandOnProcessingThread([this, PatchSettings = FModulationPatchSettings(*InPatch)]() mutable
