Remove Private Audio Modulation folder from Audio Modulation Editor Build Module

diff --git a/Engine/Plugins/Runtime/AudioModulation/Source/AudioModulationEditor/AudioModulationEditor.build.cs b/Engine/Plugins/Runtime/AudioModulation/Source/AudioModulationEditor/AudioModulationEditor.build.cs
--- a/Engine/Plugins/Runtime/AudioModulation/Source/AudioModulationEditor/AudioModulationEditor.build.cs
+++ b/Engine/Plugins/Runtime/AudioModulation/Source/AudioModulationEditor/AudioModulationEditor.build.cs
@@ -37,13 +37,6 @@ public class AudioModulationEditor : ModuleRules
 			}
 		);
 
-		PrivateIncludePaths.AddRange(
-			new string[]
-			{
-				"AudioModulation/Private"
-			}
-		);
-
 		if (bIncludeMetaSoundSupport)
 		{
 			PrivateDependencyModuleNames.AddRange(
