Audio Modulation String operation bandaid

#rnx

diff --git a/Engine/Plugins/Runtime/AudioModulation/Source/AudioModulation/Private/SoundControlBus.cpp b/Engine/Plugins/Runtime/AudioModulation/Source/AudioModulation/Private/SoundControlBus.cpp
--- a/Engine/Plugins/Runtime/AudioModulation/Source/AudioModulation/Private/SoundControlBus.cpp
+++ b/Engine/Plugins/Runtime/AudioModulation/Source/AudioModulation/Private/SoundControlBus.cpp
@@ -146,7 +146,12 @@ TSharedPtr<Audio::IProxyData> USoundControlBus::CreateProxyData(const Audio::FPr
 
 const Audio::FModulationParameter& USoundControlBus::GetOutputParameter() const
 {
-	const FString Breadcrumb = FString::Format(TEXT("{0} '{1}'"), { *GetClass()->GetName(), *GetName() });
-	return AudioModulation::GetOrRegisterParameter(Parameter, Breadcrumb);
+#if !UE_BUILD_SHIPPING
+	TStringBuilder<128> StringBuilder;
+	StringBuilder.Append(*GetClass()->GetName()).Append(" '").Append(*GetName()).Append("'");
+	return AudioModulation::GetOrRegisterParameter(Parameter, StringBuilder.ToString());
+#else
+	return AudioModulation::GetOrRegisterParameter(Parameter, FString());
+#endif
 }
 
diff --git a/Engine/Plugins/Runtime/AudioModulation/Source/AudioModulation/Private/SoundControlBusMixProxy.cpp b/Engine/Plugins/Runtime/AudioModulation/Source/AudioModulation/Private/SoundControlBusMixProxy.cpp
--- a/Engine/Plugins/Runtime/AudioModulation/Source/AudioModulation/Private/SoundControlBusMixProxy.cpp
+++ b/Engine/Plugins/Runtime/AudioModulation/Source/AudioModulation/Private/SoundControlBusMixProxy.cpp
@@ -20,7 +20,7 @@ namespace AudioModulation
 	const FBusMixId InvalidBusMixId = INDEX_NONE;
 
 	FModulatorBusMixStageSettings::FModulatorBusMixStageSettings(const FSoundControlBusMixStage& InStage)
-		: TModulatorBase<FBusId>(InStage.Bus->GetName(), InStage.Bus->GetUniqueID())
+		: TModulatorBase<FBusId>(InStage.Bus->GetFName(), InStage.Bus->GetUniqueID())
 		, Address(InStage.Bus->Address)
 		, ParamClassId(INDEX_NONE)
 		, ParamId(INDEX_NONE)
@@ -50,7 +50,7 @@ namespace AudioModulation
 	}
 
 	FModulatorBusMixSettings::FModulatorBusMixSettings(const USoundControlBusMix& InBusMix)
-		: TModulatorBase<FBusMixId>(InBusMix.GetName(), InBusMix.GetUniqueID())
+		: TModulatorBase<FBusMixId>(InBusMix.GetFName(), InBusMix.GetUniqueID())
 	{
 		for (const FSoundControlBusMixStage& Stage : InBusMix.MixStages)
 		{
diff --git a/Engine/Plugins/Runtime/AudioModulation/Source/AudioModulation/Private/SoundControlBusProxy.h b/Engine/Plugins/Runtime/AudioModulation/Source/AudioModulation/Private/SoundControlBusProxy.h
--- a/Engine/Plugins/Runtime/AudioModulation/Source/AudioModulation/Private/SoundControlBusProxy.h
+++ b/Engine/Plugins/Runtime/AudioModulation/Source/AudioModulation/Private/SoundControlBusProxy.h
@@ -27,7 +27,7 @@ namespace AudioModulation
 		Audio::FModulationParameter OutputParameter;
 
 		FControlBusSettings(const USoundControlBus& InBus)
-			: TModulatorBase<FBusId>(InBus.GetName(), InBus.GetUniqueID())
+			: TModulatorBase<FBusId>(InBus.GetFName(), InBus.GetUniqueID())
 			, bBypass(InBus.bBypass)
 			, DefaultValue(InBus.GetDefaultNormalizedValue())
 			, MixFunction(InBus.GetMixFunction())
diff --git a/Engine/Plugins/Runtime/AudioModulation/Source/AudioModulation/Private/SoundModulationGeneratorProxy.h b/Engine/Plugins/Runtime/AudioModulation/Source/AudioModulation/Private/SoundModulationGeneratorProxy.h
--- a/Engine/Plugins/Runtime/AudioModulation/Source/AudioModulation/Private/SoundModulationGeneratorProxy.h
+++ b/Engine/Plugins/Runtime/AudioModulation/Source/AudioModulation/Private/SoundModulationGeneratorProxy.h
@@ -30,7 +30,7 @@ namespace AudioModulation
 		FModulationGeneratorSettings() = default;
 
 		FModulationGeneratorSettings(const USoundModulationGenerator& InGenerator)
-			: TModulatorBase<FGeneratorId>(InGenerator.GetName(), InGenerator.GetUniqueID())
+			: TModulatorBase<FGeneratorId>(InGenerator.GetFName(), InGenerator.GetUniqueID())
 			, Generator(InGenerator.CreateInstance())
 		{
 		}
diff --git a/Engine/Plugins/Runtime/AudioModulation/Source/AudioModulation/Private/SoundModulationPatchProxy.h b/Engine/Plugins/Runtime/AudioModulation/Source/AudioModulation/Private/SoundModulationPatchProxy.h
--- a/Engine/Plugins/Runtime/AudioModulation/Source/AudioModulation/Private/SoundModulationPatchProxy.h
+++ b/Engine/Plugins/Runtime/AudioModulation/Source/AudioModulation/Private/SoundModulationPatchProxy.h
@@ -83,7 +83,7 @@ namespace AudioModulation
 		FModulationPatchSettings(const FModulationPatchSettings& InPatchSettings) = default;
 
 		FModulationPatchSettings(const USoundModulationPatch& InPatch)
-			: TModulatorBase<FPatchId>(InPatch.GetName(), InPatch.GetUniqueID())
+			: TModulatorBase<FPatchId>(InPatch.GetFName(), InPatch.GetUniqueID())
 			, bBypass(InPatch.PatchSettings.bBypass)
 			, OutputParameter(InPatch.GetOutputParameter())
 		{
diff --git a/Engine/Plugins/Runtime/AudioModulation/Source/AudioModulation/Private/SoundModulationProxy.h b/Engine/Plugins/Runtime/AudioModulation/Source/AudioModulation/Private/SoundModulationProxy.h
--- a/Engine/Plugins/Runtime/AudioModulation/Source/AudioModulation/Private/SoundModulationProxy.h
+++ b/Engine/Plugins/Runtime/AudioModulation/Source/AudioModulation/Private/SoundModulationProxy.h
@@ -246,13 +246,21 @@ namespace AudioModulation
 		{
 		}
 
-		TModulatorBase(const FString& InName, const uint32 InId)
+		TModulatorBase(const FName InName, const uint32 InId)
 			: Id(static_cast<IdType>(InId))
 #if !UE_BUILD_SHIPPING
-			, Name(InName)
+			, Name(InName.ToString())
 #endif // !UE_BUILD_SHIPPING
 		{
 		}
+		
+		TModulatorBase(const FString& InName, const uint32 InId)
+        			: Id(static_cast<IdType>(InId))
+        #if !UE_BUILD_SHIPPING
+        			, Name(InName)
+        #endif // !UE_BUILD_SHIPPING
+        {
+        }
 
 		virtual ~TModulatorBase() = default;
 
