Sound Modulation GT/BP Watcher

diff --git a/Engine/Plugins/Runtime/AudioModulation/AudioModulation.uplugin b/Engine/Plugins/Runtime/AudioModulation/AudioModulation.uplugin
--- a/Engine/Plugins/Runtime/AudioModulation/AudioModulation.uplugin
+++ b/Engine/Plugins/Runtime/AudioModulation/AudioModulation.uplugin
@@ -25,7 +25,12 @@
 			"Name": "AudioModulationEditor",
 			"Type": "Editor",
 			"LoadingPhase": "PostEngineInit"
-		}
+		},
+		{
+			"Name": "AudioModulationTest",
+			"Type": "Runtime",
+			"LoadingPhase": "PreDefault"
+		},
 	],
 	"Plugins": [
 		{
diff --git a/Engine/Plugins/Runtime/AudioModulation/Source/AudioModulation/Private/AudioModulation.cpp b/Engine/Plugins/Runtime/AudioModulation/Source/AudioModulation/Private/AudioModulation.cpp
--- a/Engine/Plugins/Runtime/AudioModulation/Source/AudioModulation/Private/AudioModulation.cpp
+++ b/Engine/Plugins/Runtime/AudioModulation/Source/AudioModulation/Private/AudioModulation.cpp
@@ -48,7 +48,11 @@ namespace AudioModulation
 
 	void FAudioModulationManager::ActivateBus(const USoundControlBus& InBus)
 	{
+		UE_LOG(LogAudioModulation, Warning, TEXT("FAudioModulationManager::ActivateBus is deprecated. Use USoundModulationWatcher instead"));
+
+		PRAGMA_DISABLE_DEPRECATION_WARNINGS
 		ModSystem->ActivateBus(InBus);
+		PRAGMA_ENABLE_DEPRECATION_WARNINGS
 	}
 
 	void FAudioModulationManager::ActivateBusMix(const USoundControlBusMix& InBusMix)
@@ -58,12 +62,20 @@ namespace AudioModulation
 
 	void FAudioModulationManager::ActivateGenerator(const USoundModulationGenerator& InGenerator)
 	{
+		UE_LOG(LogAudioModulation, Warning, TEXT("FAudioModulationManager::ActivateGenerator is deprecated. Use USoundModulationWatcher instead"));
+
+		PRAGMA_DISABLE_DEPRECATION_WARNINGS
 		ModSystem->ActivateGenerator(InGenerator);
+		PRAGMA_ENABLE_DEPRECATION_WARNINGS
 	}
 
 	void FAudioModulationManager::DeactivateBus(const USoundControlBus& InBus)
 	{
+		UE_LOG(LogAudioModulation, Warning, TEXT("FAudioModulationManager::DeactivateBus is deprecated. Use USoundModulationWatcher instead"));
+
+		PRAGMA_DISABLE_DEPRECATION_WARNINGS
 		ModSystem->DeactivateBus(InBus);
+		PRAGMA_ENABLE_DEPRECATION_WARNINGS
 	}
 
 	void FAudioModulationManager::DeactivateBusMix(const USoundControlBusMix& InBusMix)
@@ -78,7 +90,11 @@ namespace AudioModulation
 
 	void FAudioModulationManager::DeactivateGenerator(const USoundModulationGenerator& InGenerator)
 	{
+		UE_LOG(LogAudioModulation, Warning, TEXT("FAudioModulationManager::DeactivateGenerator is deprecated. Use USoundModulationWatcher instead"));
+
+		PRAGMA_DISABLE_DEPRECATION_WARNINGS
 		ModSystem->DeactivateGenerator(InGenerator);
+		PRAGMA_ENABLE_DEPRECATION_WARNINGS
 	}
 
 #if !UE_BUILD_SHIPPING
diff --git a/Engine/Plugins/Runtime/AudioModulation/Source/AudioModulation/Private/AudioModulationStatics.cpp b/Engine/Plugins/Runtime/AudioModulation/Source/AudioModulation/Private/AudioModulationStatics.cpp
--- a/Engine/Plugins/Runtime/AudioModulation/Source/AudioModulation/Private/AudioModulationStatics.cpp
+++ b/Engine/Plugins/Runtime/AudioModulation/Source/AudioModulation/Private/AudioModulationStatics.cpp
@@ -15,6 +15,7 @@
 #include "Misc/Paths.h"
 #include "SoundControlBus.h"
 #include "SoundControlBusMix.h"
+#include "SoundModulationWatcher.h"
 
 #include UE_INLINE_GENERATED_CPP_BY_NAME(AudioModulationStatics)
 
@@ -121,7 +122,9 @@ void UAudioModulationStatics::ActivateBus(const UObject* WorldContextObject, USo
 	UWorld* World = GetAudioWorld(WorldContextObject);
 	if (AudioModulation::FAudioModulationManager* ModSystem = GetModulation(World))
 	{
+		PRAGMA_DISABLE_DEPRECATION_WARNINGS
 		ModSystem->ActivateBus(*Bus);
+		PRAGMA_ENABLE_DEPRECATION_WARNINGS
 	}
 }
 
@@ -144,7 +147,9 @@ void UAudioModulationStatics::ActivateGenerator(const UObject* WorldContextObjec
 	{
 		if (Generator)
 		{
+			PRAGMA_DISABLE_DEPRECATION_WARNINGS
 			ModSystem->ActivateGenerator(*Generator);
+			PRAGMA_ENABLE_DEPRECATION_WARNINGS
 		}
 	}
 }
@@ -207,7 +212,9 @@ USoundControlBus* UAudioModulationStatics::CreateBus(UObject* WorldContextObject
 	{
 		if (AudioModulation::FAudioModulationManager* ModSystem = GetModulation(World))
 		{
+			PRAGMA_DISABLE_DEPRECATION_WARNINGS
 			ModSystem->ActivateBus(*NewBus);
+			PRAGMA_ENABLE_DEPRECATION_WARNINGS
 		}
 	}
 
@@ -328,6 +335,17 @@ USoundModulationGeneratorADEnvelope* UAudioModulationStatics::CreateADEnvelopeGe
 	return NewGenerator;
 }
 
+USoundModulationWatcher* UAudioModulationStatics::CreateModulationWatcher(UObject* WorldContextObject, FName Name, USoundModulatorBase* Modulator)
+{
+	if (USoundModulationWatcher* NewWatcher = NewObject<USoundModulationWatcher>(WorldContextObject, Name))
+	{
+		NewWatcher->SetModulator(Modulator);
+		return NewWatcher;
+	}
+
+	return nullptr;
+}
+
 void UAudioModulationStatics::DeactivateBus(const UObject* WorldContextObject, USoundControlBus* Bus)
 {
 	if (Bus)
@@ -335,7 +353,9 @@ void UAudioModulationStatics::DeactivateBus(const UObject* WorldContextObject, U
 		UWorld* World = GetAudioWorld(WorldContextObject);
 		if (AudioModulation::FAudioModulationManager* ModSystem = GetModulation(World))
 		{
+			PRAGMA_DISABLE_DEPRECATION_WARNINGS
 			ModSystem->DeactivateBus(*Bus);
+			PRAGMA_ENABLE_DEPRECATION_WARNINGS
 		}
 	}
 }
@@ -359,7 +379,9 @@ void UAudioModulationStatics::DeactivateGenerator(const UObject* WorldContextObj
 	{
 		if (Generator)
 		{
+			PRAGMA_DISABLE_DEPRECATION_WARNINGS
 			ModSystem->DeactivateGenerator(*Generator);
+			PRAGMA_ENABLE_DEPRECATION_WARNINGS
 		}
 	}
 }
diff --git a/Engine/Plugins/Runtime/AudioModulation/Source/AudioModulation/Private/AudioModulationSystem.h b/Engine/Plugins/Runtime/AudioModulation/Source/AudioModulation/Private/AudioModulationSystem.h
--- a/Engine/Plugins/Runtime/AudioModulation/Source/AudioModulation/Private/AudioModulationSystem.h
+++ b/Engine/Plugins/Runtime/AudioModulation/Source/AudioModulation/Private/AudioModulationSystem.h
@@ -67,18 +67,25 @@ namespace AudioModulation
 	public:
 		void Initialize(const FAudioPluginInitializationParams& InitializationParams);
 
+		UE_DEPRECATED(5.4, "Activation of modulators in this manner is now deprecated. Use USoundModulationWatchers to safety activate and track a given modulator")
 		void ActivateBus(const USoundControlBus& InBus);
+
 		void ActivateBusMix(FModulatorBusMixSettings&& InSettings);
 		void ActivateBusMix(const USoundControlBusMix& InBusMix);
+
+		UE_DEPRECATED(5.4, "Activation of modulators in this manner is now deprecated. Use USoundModulationWatchers to safety activate and track a given modulator")
 		void ActivateGenerator(const USoundModulationGenerator& InGenerator);
 
-		/**
-		 * Deactivates respectively typed (i.e. BusMix, Bus, Generator, etc.) object proxy if no longer referenced.
-		 * If still referenced, will wait until references are finished before destroying.
-		 */
+		UE_DEPRECATED(5.4, "Deactivation of modulators in this manner is now deprecated. Use USoundModulationWatchers to safety activate and track a given modulator")
 		void DeactivateBus(const USoundControlBus& InBus);
+
+		/** Deactivates given bus mix */
 		void DeactivateBusMix(const USoundControlBusMix& InBusMix);
+
+		/** Deactivates all bus mixes */
 		void DeactivateAllBusMixes();
+
+		UE_DEPRECATED(5.4, "Deactivation of modulators in this manner is now deprecated. Use USoundModulationWatchers to safety activate and track a given modulator")
 		void DeactivateGenerator(const USoundModulationGenerator& InGenerator);
 
 		void ProcessModulators(const double InElapsed);
diff --git a/Engine/Plugins/Runtime/AudioModulation/Source/AudioModulation/Private/SoundControlBus.cpp b/Engine/Plugins/Runtime/AudioModulation/Source/AudioModulation/Private/SoundControlBus.cpp
--- a/Engine/Plugins/Runtime/AudioModulation/Source/AudioModulation/Private/SoundControlBus.cpp
+++ b/Engine/Plugins/Runtime/AudioModulation/Source/AudioModulation/Private/SoundControlBus.cpp
@@ -119,7 +119,9 @@ void USoundControlBus::BeginDestroy()
 			{
 				FAudioModulationManager* Modulation = static_cast<FAudioModulationManager*>(ModulationInterface);
 				check(Modulation);
+				PRAGMA_DISABLE_DEPRECATION_WARNINGS
 				Modulation->DeactivateBus(*this);
+				PRAGMA_ENABLE_DEPRECATION_WARNINGS
 			}
 		}
 	}
diff --git a/Engine/Plugins/Runtime/AudioModulation/Source/AudioModulation/Private/SoundModulationGenerator.cpp b/Engine/Plugins/Runtime/AudioModulation/Source/AudioModulation/Private/SoundModulationGenerator.cpp
--- a/Engine/Plugins/Runtime/AudioModulation/Source/AudioModulation/Private/SoundModulationGenerator.cpp
+++ b/Engine/Plugins/Runtime/AudioModulation/Source/AudioModulation/Private/SoundModulationGenerator.cpp
@@ -72,7 +72,9 @@ void USoundModulationGenerator::BeginDestroy()
 			{
 				FAudioModulationManager* Modulation = static_cast<FAudioModulationManager*>(ModulationInterface);
 				check(Modulation);
+				PRAGMA_DISABLE_DEPRECATION_WARNINGS
 				Modulation->DeactivateGenerator(*this);
+				PRAGMA_ENABLE_DEPRECATION_WARNINGS
 			}
 		}
 	}
diff --git a/Engine/Plugins/Runtime/AudioModulation/Source/AudioModulation/Private/SoundModulationWatcher.cpp b/Engine/Plugins/Runtime/AudioModulation/Source/AudioModulation/Private/SoundModulationWatcher.cpp
new file mode 100644
--- /dev/null
+++ b/Engine/Plugins/Runtime/AudioModulation/Source/AudioModulation/Private/SoundModulationWatcher.cpp
@@ -0,0 +1,117 @@
+// Copyright Epic Games, Inc. All Rights Reserved.
+#include "SoundModulationWatcher.h"
+
+#include "AudioDeviceHandle.h"
+#include "AudioModulation.h"
+#include "Engine/Engine.h"
+#include "Engine/EngineBaseTypes.h"
+#include "Engine/World.h"
+
+#include UE_INLINE_GENERATED_CPP_BY_NAME(SoundModulationWatcher)
+
+
+namespace AudioModulation::WatcherPrivate
+{
+	FAudioDeviceHandle GetAudioDevice(const USoundModulationWatcher& InWatcher)
+	{
+		if (GEngine)
+		{
+			if (GEngine->UseSound())
+			{
+				UWorld* World = GEngine->GetWorldFromContextObject(&InWatcher, EGetWorldErrorMode::ReturnNull);
+				if (World && World->bAllowAudioPlayback && !World->IsNetMode(NM_DedicatedServer))
+				{
+					return World->GetAudioDevice();
+				}
+				else
+				{
+					return GEngine->GetMainAudioDevice();
+				}
+			}
+		}
+
+		return { };
+	}
+
+	FAudioModulationManager* GetModulationManager(const USoundModulationWatcher& InWatcher)
+	{
+		FAudioDeviceHandle AudioDevice = WatcherPrivate::GetAudioDevice(InWatcher);
+		if (AudioDevice.IsValid() && AudioDevice->IsModulationPluginEnabled())
+		{
+			if (IAudioModulationManager* ModulationInterface = AudioDevice->ModulationInterface.Get())
+			{
+				return static_cast<FAudioModulationManager*>(ModulationInterface);
+			}
+		}
+
+		return nullptr;
+	}
+} // AudioModulation::WatcherPrivate
+
+bool USoundModulationWatcher::ClearModulator()
+{
+	if (Modulator)
+	{
+		Modulator = nullptr;
+		Destination.UpdateModulators(TSet<const USoundModulatorBase*>{ });
+		return true;
+	}
+
+	return false;
+}
+
+const USoundModulatorBase* USoundModulationWatcher::GetModulator() const
+{
+	return Modulator;
+}
+
+float USoundModulationWatcher::GetValue() const
+{
+	using namespace AudioModulation;
+
+	if (Modulator)
+	{
+		if (FAudioModulationManager* Modulation = WatcherPrivate::GetModulationManager(*this))
+		{
+			return Modulation->GetModulatorValueThreadSafe(Modulator->GetUniqueID());
+		}
+	}
+
+	return 1.0f;
+}
+
+void USoundModulationWatcher::PostInitProperties()
+{
+	using namespace AudioModulation;
+
+	Super::PostInitProperties();
+
+	if (USoundModulationWatcher::StaticClass()->GetDefaultObject() == this)
+	{
+		return;
+	}
+
+	FAudioDeviceHandle AudioDevice = WatcherPrivate::GetAudioDevice(*this);
+	if (AudioDevice.IsValid())
+	{
+		constexpr bool bIsBuffered = false;
+		Destination.Init(AudioDevice.GetDeviceID(), bIsBuffered);
+	}
+}
+
+bool USoundModulationWatcher::SetModulator(const USoundModulatorBase* InModulator)
+{
+	if (InModulator == Modulator)
+	{
+		return true;
+	}
+
+	if (InModulator)
+	{
+		Modulator = InModulator;
+		Destination.UpdateModulators(TSet<const USoundModulatorBase*>{ InModulator });
+		return true;
+	}
+
+	return ClearModulator();
+}
diff --git a/Engine/Plugins/Runtime/AudioModulation/Source/AudioModulation/Public/AudioModulation.h b/Engine/Plugins/Runtime/AudioModulation/Source/AudioModulation/Public/AudioModulation.h
--- a/Engine/Plugins/Runtime/AudioModulation/Source/AudioModulation/Public/AudioModulation.h
+++ b/Engine/Plugins/Runtime/AudioModulation/Source/AudioModulation/Public/AudioModulation.h
@@ -59,13 +59,21 @@ namespace AudioModulation
 	//~ End IAudioModulationManager implementation
 
 	public:
+		UE_DEPRECATED(5.4, "Deactivation of modulators in this manner is now deprecated. Use USoundModulationWatchers to safety activate and track a given modulator")
 		void ActivateBus(const USoundControlBus& InBus);
+
 		void ActivateBusMix(const USoundControlBusMix& InBusMix);
+
+		UE_DEPRECATED(5.4, "Deactivation of modulators in this manner is now deprecated. Use USoundModulationWatchers to safety activate and track a given modulator")
 		void ActivateGenerator(const USoundModulationGenerator& InGenerator);
 
+		UE_DEPRECATED(5.4, "Deactivation of modulators in this manner is now deprecated. Use USoundModulationWatchers to safety activate and track a given modulator")
 		void DeactivateBus(const USoundControlBus& InBus);
+
 		void DeactivateBusMix(const USoundControlBusMix& InBusMix);
 		void DeactivateAllBusMixes();
+
+		UE_DEPRECATED(5.4, "Deactivation of modulators in this manner is now deprecated. Use USoundModulationWatchers to safety activate and track a given modulator")
 		void DeactivateGenerator(const USoundModulationGenerator& InGenerator);
 
 		void SaveMixToProfile(const USoundControlBusMix& InBusMix, const int32 InProfileIndex);
diff --git a/Engine/Plugins/Runtime/AudioModulation/Source/AudioModulation/Public/AudioModulationStatics.h b/Engine/Plugins/Runtime/AudioModulation/Source/AudioModulation/Public/AudioModulationStatics.h
--- a/Engine/Plugins/Runtime/AudioModulation/Source/AudioModulation/Public/AudioModulationStatics.h
+++ b/Engine/Plugins/Runtime/AudioModulation/Source/AudioModulation/Public/AudioModulationStatics.h
@@ -14,6 +14,8 @@
 #include "AudioModulationStatics.generated.h"
 
 // Forward Declarations
+class USoundModulationWatcher;
+
 namespace AudioModulation
 {
 	class FAudioModulationManager;
@@ -37,7 +39,8 @@ public:
 	 */
 	static AudioModulation::FAudioModulationManager* GetModulation(UWorld* World);
 
-	/** Manually activates a modulation bus. If called, deactivation will only occur
+	/** SOFT DEPRECATED: Use CreateModulationWatcher and store resulting watcher instead!
+	 * Manually activates a modulation bus. If called, deactivation will only occur
 	 * if bus is manually deactivated or destroyed (i.e. will not deactivate
 	 * when all references become inactive).
 	 */
@@ -58,7 +61,8 @@ public:
 	)
 	static void ActivateBusMix(const UObject* WorldContextObject, USoundControlBusMix* Mix);
 
-	/** Manually activates a modulation generator. If called, deactivation will only occur
+	/** SOFT DEPRECATED: Use CreateModulationWatcher and store resulting watcher instead!
+	 * Manually activates a modulation generator. If called, deactivation will only occur
 	 * if generator is manually deactivated and not referenced or destroyed (i.e. will not deactivate
 	 * when all references become inactive).
 	 * @param Modulator - Modulator to activate
@@ -173,12 +177,22 @@ public:
 		WorldContext = "WorldContextObject",
 		Keywords = "make modulator")
 	)
-		static UPARAM(DisplayName = "Generator") USoundModulationGeneratorADEnvelope* CreateADEnvelopeGenerator(
-			UObject* WorldContextObject,
-			FName Name,
-			const FSoundModulationADEnvelopeParams& Params);
+	static UPARAM(DisplayName = "Generator") USoundModulationGeneratorADEnvelope* CreateADEnvelopeGenerator(
+		UObject* WorldContextObject,
+		FName Name,
+		const FSoundModulationADEnvelopeParams& Params);
+
+	/** Creates a modulation watcher, which activates the given modulator and provides a function to retrieve
+	 * the last value computed of the given modulator on the modulation processing thread.
+	 */
+	UFUNCTION(BlueprintCallable, Category = "Audio|Modulation", meta = (WorldContext = "WorldContextObject", Keywords = "modulator watch watcher"))
+	static UPARAM(DisplayName = "Watcher") USoundModulationWatcher* CreateModulationWatcher(
+		UObject* WorldContextObject,
+		FName Name,
+		USoundModulatorBase* Modulator);
 
-	/** Deactivates a bus. Does nothing if the provided bus is already inactive.
+	/** SOFT DEPRECATED: Use CreateModulationWatcher and store resulting watcher instead!
+	 * Deactivates a bus. Does nothing if the provided bus is already inactive.
 	 * @param Bus - Scope of modulator
 	 */
 	UFUNCTION(BlueprintCallable, Category = "Audio|Modulation", DisplayName = "Deactivate Control Bus", meta = (
@@ -196,7 +210,8 @@ public:
 	)
 	static void DeactivateBusMix(const UObject* WorldContextObject, USoundControlBusMix* Mix);
 
-	/** Deactivates a modulation generator. Does nothing if an instance of the provided generator is already inactive.
+	/** SOFT DEPRECATED: Use CreateModulationWatcher and store resulting watcher instead!
+	 * Deactivates a modulation generator. Does nothing if an instance of the provided generator is already inactive.
 	 * @param Generator - Generator to activate
 	 * @param Scope - Scope of modulator
 	 */
diff --git a/Engine/Plugins/Runtime/AudioModulation/Source/AudioModulation/Public/SoundModulationWatcher.h b/Engine/Plugins/Runtime/AudioModulation/Source/AudioModulation/Public/SoundModulationWatcher.h
new file mode 100644
--- /dev/null
+++ b/Engine/Plugins/Runtime/AudioModulation/Source/AudioModulation/Public/SoundModulationWatcher.h
@@ -0,0 +1,41 @@
+// Copyright Epic Games, Inc. All Rights Reserved.
+#pragma once
+
+#include "IAudioModulation.h"
+
+#include "AudioThread.h"
+#include "Sound/SoundModulationDestination.h"
+
+#include "SoundModulationWatcher.generated.h"
+
+
+UCLASS(config = Engine, editinlinenew, BlueprintType)
+class AUDIOMODULATION_API USoundModulationWatcher : public UObject
+{
+	GENERATED_BODY()
+
+private:
+	UPROPERTY()
+	TObjectPtr<const USoundModulatorBase> Modulator;
+
+	Audio::FModulationDestination Destination;
+
+public:
+	virtual void PostInitProperties() override;
+
+	// Returns true if a modulator was being watched and has been cleared.
+	UFUNCTION(BlueprintCallable, Category = "Audio|Modulation", DisplayName = "Clear Modulator")
+	UPARAM(DisplayName = "Modulator Cleared") bool ClearModulator();
+
+	// Returns currently set modulator being watched.
+	UFUNCTION(BlueprintPure, Category = "Audio|Modulation")
+	UPARAM(DisplayName = "Modulator") const USoundModulatorBase* GetModulator() const;
+
+	// Returns the last calculated modulator value sampled from the thread controls are processed on.
+	UFUNCTION(BlueprintPure, Category = "Audio|Modulation", DisplayName = "Get Watched Modulator Value")
+	UPARAM(DisplayName = "Value") float GetValue() const;
+
+	// Returns true if modulator was set to new value or was already set to provided value.
+	UFUNCTION(BlueprintCallable, Category = "Audio|Modulation", DisplayName = "Set Watched Modulator")
+	UPARAM(DisplayName = "Modulator Set") bool SetModulator(const USoundModulatorBase* InModulator);
+};
diff --git a/Engine/Plugins/Runtime/AudioModulation/Source/AudioModulationTest/AudioModulationTest.Build.cs b/Engine/Plugins/Runtime/AudioModulation/Source/AudioModulationTest/AudioModulationTest.Build.cs
new file mode 100644
--- /dev/null
+++ b/Engine/Plugins/Runtime/AudioModulation/Source/AudioModulationTest/AudioModulationTest.Build.cs
@@ -0,0 +1,34 @@
+// Copyright Epic Games, Inc. All Rights Reserved.
+
+using UnrealBuildTool;
+using System.IO;
+
+namespace UnrealBuildTool.Rules
+{
+	public class AudioModulationTest : ModuleRules
+	{
+		public AudioModulationTest(ReadOnlyTargetRules Target) : base(Target)
+		{
+			PrivateDependencyModuleNames.AddRange
+			(
+				new string[]
+				{
+					"Core",
+					"CoreUObject",
+					"Engine",
+					"AudioExtensions",
+					"AudioMixer",
+					"Projects"
+				}
+			);
+
+			PublicDependencyModuleNames.AddRange
+			(
+				new string[]
+				{
+					"AudioModulation",
+				}
+			);
+		}
+	}
+}
diff --git a/Engine/Plugins/Runtime/AudioModulation/Source/AudioModulationTest/Private/AudioModulationTestModule.cpp b/Engine/Plugins/Runtime/AudioModulation/Source/AudioModulationTest/Private/AudioModulationTestModule.cpp
new file mode 100644
--- /dev/null
+++ b/Engine/Plugins/Runtime/AudioModulation/Source/AudioModulationTest/Private/AudioModulationTestModule.cpp
@@ -0,0 +1,13 @@
+// Copyright Epic Games, Inc. All Rights Reserved.
+
+#include "Modules/ModuleInterface.h"
+#include "Modules/ModuleManager.h"
+
+namespace AudioModulation 
+{
+	class AUDIOMODULATIONTEST_API FTestModule : public IModuleInterface
+	{
+	};
+} // namespace AudioModulation
+
+IMPLEMENT_MODULE(AudioModulation::FTestModule, AudioModulationTest);
diff --git a/Engine/Plugins/Runtime/AudioModulation/Source/AudioModulationTest/Private/AudioModulationWatcherTest.cpp b/Engine/Plugins/Runtime/AudioModulation/Source/AudioModulationTest/Private/AudioModulationWatcherTest.cpp
new file mode 100644
--- /dev/null
+++ b/Engine/Plugins/Runtime/AudioModulation/Source/AudioModulationTest/Private/AudioModulationWatcherTest.cpp
@@ -0,0 +1,208 @@
+// Copyright Epic Games, Inc. All Rights Reserved.
+
+#if WITH_DEV_AUTOMATION_TESTS
+
+#include "AudioModulationStatics.h"
+#include "Engine/Engine.h"
+#include "Interfaces/IPluginManager.h"
+#include "Misc/AutomationTest.h"
+#include "Misc/Paths.h"
+#include "SoundControlBus.h"
+#include "Sound/SoundAttenuation.h"
+#include "SoundModulationWatcher.h"
+#include "Templates/SharedPointer.h"
+#include "Tests/AutomationCommon.h"
+
+
+namespace AudioModulation::TestPrivate
+{
+	FString GetPluginContentDirectory()
+	{
+		TSharedPtr<IPlugin> Plugin = IPluginManager::Get().FindPlugin(TEXT("AudioModulation"));
+		if (ensure(Plugin.IsValid()))
+		{
+			return Plugin->GetContentDir();
+		}
+		return FString();
+	}
+
+	FString GetPathToTestFilesDir()
+	{
+		FString OutPath =  FPaths::Combine(GetPluginContentDirectory(), TEXT("Test"));
+
+		OutPath = FPaths::ConvertRelativePathToFull(OutPath);
+		FPaths::NormalizeDirectoryName(OutPath);
+		
+		return OutPath;
+	}
+
+	FString GetPathToGeneratedFilesDir()
+	{
+		FString OutPath = FPaths::Combine(FPaths::ProjectSavedDir(), TEXT("AudioModulation"));
+
+		OutPath = FPaths::ConvertRelativePathToFull(OutPath);
+		FPaths::NormalizeDirectoryName(OutPath);
+		
+		return OutPath;
+	}
+
+	FString GetPathToGeneratedAssetsDir()
+	{
+		FString OutPath = TEXT("/Game/AudioModulation/Generated/");
+		FPaths::NormalizeDirectoryName(OutPath);
+		return OutPath;
+	}
+} // AudioModulation::TestPrivate
+
+// Tests return value of Watcher's GetValue command against the provided value
+DEFINE_LATENT_AUTOMATION_COMMAND_THREE_PARAMETER(FSoundModulationWatcherGetValueLatentCommand, FAutomationTestBase&, Test, USoundModulationWatcher*, Watcher, float, TestValue);
+
+bool FSoundModulationWatcherGetValueLatentCommand::Update()
+{
+	if (Watcher)
+	{
+		const float WatcherValue = Watcher->GetValue();
+		Test.AddErrorIfFalse(FMath::IsNearlyEqual(WatcherValue, TestValue), FString::Format(TEXT("Modulation Watcher is reporting value {0} but expected value is {1}"), { WatcherValue, TestValue }));
+		return true;
+	}
+
+	return true;
+}
+
+DEFINE_LATENT_AUTOMATION_COMMAND_TWO_PARAMETER(FSoundModulationWatcherClearModulatorLatentCommand, FAutomationTestBase&, Test, USoundModulationWatcher*, Watcher);
+
+bool FSoundModulationWatcherClearModulatorLatentCommand::Update()
+{
+	if (Watcher)
+	{
+		Watcher->ClearModulator();
+		return true;
+	}
+
+	return false;
+}
+
+DEFINE_LATENT_AUTOMATION_COMMAND_THREE_PARAMETER(FSoundModulationWatcherSetModulatorLatentCommand, FAutomationTestBase&, Test, USoundModulationWatcher*, Watcher, USoundModulatorBase*, Modulator);
+
+bool FSoundModulationWatcherSetModulatorLatentCommand::Update()
+{
+	if (Watcher)
+	{
+		Watcher->SetModulator(Modulator);
+		return true;
+	}
+
+	return false;
+}
+
+DEFINE_LATENT_AUTOMATION_COMMAND_ONE_PARAMETER(FSoundModulationRemoveFromRootLatentCommand, UObject*, ToRemove);
+
+bool FSoundModulationRemoveFromRootLatentCommand::Update()
+{
+	if (ToRemove)
+	{
+		ToRemove->RemoveFromRoot();
+		return true;
+	}
+
+	return false;
+}
+
+
+// This test creates a Modulation bus, watcher, and mix.  The bus's value is then tested both before and after the mix is activated.
+IMPLEMENT_SIMPLE_AUTOMATION_TEST(FSoundModulationWatcherGetValueTest, "Audio.Modulation.Watcher.GetValue", EAutomationTestFlags::EditorContext | EAutomationTestFlags::EngineFilter)
+bool FSoundModulationWatcherGetValueTest::RunTest(const FString& Parameters)
+{
+	using namespace Audio;
+
+	UWorld* TestWorld = nullptr;
+	if (GEngine)
+	{
+		const TIndirectArray<FWorldContext>& WorldContexts = GEngine->GetWorldContexts();
+		for (const FWorldContext& Context : WorldContexts)
+		{
+			if (Context.World() != nullptr)
+			{
+				if (Context.WorldType != EWorldType::Inactive && Context.WorldType != EWorldType::None)
+				{
+					TestWorld = Context.World();
+				}
+			}
+		}
+	}
+
+	if (!TestWorld)
+	{
+		return false;
+	}
+
+	USoundModulationWatcher* TestWatcher = NewObject<USoundModulationWatcher>();
+	TestWatcher->AddToRoot();
+
+	USoundControlBus* TestBus = NewObject<USoundControlBus>();
+	TestBus->AddToRoot();
+
+	USoundControlBus* TestBus2 = NewObject<USoundControlBus>();
+	TestBus2->AddToRoot();
+
+	TestWatcher->SetModulator(TestBus);
+
+	constexpr float TestMixValue = 0.25f;
+	FSoundControlBusMixStage Stage;
+	Stage.Bus = TestBus;
+	Stage.Value = FSoundModulationMixValue { TestMixValue, 0.1f /* AttackTime */, 0.1f /* ReleaseTime */ };
+
+	const float WatcherValue = TestWatcher->GetValue();
+	AddErrorIfFalse(FMath::IsNearlyEqual(WatcherValue, 1.0f), FString::Format(TEXT("Modulation Watcher is reporting value {0} but initial expected value is 1.0"), { WatcherValue }));
+
+	constexpr bool bActivate = true;
+	USoundControlBusMix* TestMix = UAudioModulationStatics::CreateBusMix(TestWorld, "TestMix", { Stage }, bActivate);
+	AddErrorIfFalse(TestMix != nullptr, TEXT("Failed to create control bus mix for Modulation Watcher 'GetValue' test"));
+
+	constexpr float TestMix2Value = 0.33f;
+	Stage.Bus = TestBus2;
+	Stage.Value = FSoundModulationMixValue { TestMix2Value, 0.1f /* AttackTime */, 0.1f /* ReleaseTime */ };
+	USoundControlBusMix* TestMix2 = UAudioModulationStatics::CreateBusMix(TestWorld, "TestMix2", { Stage }, bActivate);
+	AddErrorIfFalse(TestMix2 != nullptr, TEXT("Failed to create control bus mix for Modulation Watcher 'GetValue' test"));
+
+	auto Cleanup = [&]()
+	{
+		// Clean up, clean up, everybody everywhere
+		ADD_LATENT_AUTOMATION_COMMAND(FSoundModulationRemoveFromRootLatentCommand(TestWatcher));
+		ADD_LATENT_AUTOMATION_COMMAND(FSoundModulationRemoveFromRootLatentCommand(TestBus));
+		ADD_LATENT_AUTOMATION_COMMAND(FSoundModulationRemoveFromRootLatentCommand(TestBus2));
+		ADD_LATENT_AUTOMATION_COMMAND(FSoundModulationRemoveFromRootLatentCommand(TestMix));
+		ADD_LATENT_AUTOMATION_COMMAND(FSoundModulationRemoveFromRootLatentCommand(TestMix2));
+	};
+
+	if (TestMix && TestMix2)
+	{
+		TestMix->AddToRoot();
+		TestMix2->AddToRoot();
+
+		// Wait for mix to apply
+		ADD_LATENT_AUTOMATION_COMMAND(FEngineWaitLatentCommand(0.125f));
+
+		// Validate setting modulator works (value should be modified to test mix 2's value)
+		ADD_LATENT_AUTOMATION_COMMAND(FSoundModulationWatcherSetModulatorLatentCommand(*this, TestWatcher, TestBus2));
+		ADD_LATENT_AUTOMATION_COMMAND(FSoundModulationWatcherGetValueLatentCommand(*this, TestWatcher, TestMix2Value));
+
+		// Set back to initial bus for remainder of tests
+		ADD_LATENT_AUTOMATION_COMMAND(FSoundModulationWatcherSetModulatorLatentCommand(*this, TestWatcher, TestBus));
+
+		// Validate mix application mutates modulator value over time
+		ADD_LATENT_AUTOMATION_COMMAND(FSoundModulationWatcherGetValueLatentCommand(*this, TestWatcher, TestMixValue));
+
+		// Validate clearing value works
+		ADD_LATENT_AUTOMATION_COMMAND(FSoundModulationWatcherClearModulatorLatentCommand(*this, TestWatcher));
+		ADD_LATENT_AUTOMATION_COMMAND(FSoundModulationWatcherGetValueLatentCommand(*this, TestWatcher, 1.0f));
+
+		Cleanup();
+		return true;
+	}
+
+	Cleanup();
+	return false;
+}
+
+#endif // WITH_DEV_AUTOMATION_TESTS
