Backing out modulation test module until editor tests gets fixed.

diff --git a/Engine/Plugins/Runtime/AudioModulation/AudioModulation.uplugin b/Engine/Plugins/Runtime/AudioModulation/AudioModulation.uplugin
--- a/Engine/Plugins/Runtime/AudioModulation/AudioModulation.uplugin
+++ b/Engine/Plugins/Runtime/AudioModulation/AudioModulation.uplugin
@@ -25,12 +25,7 @@
 			"Name": "AudioModulationEditor",
 			"Type": "Editor",
 			"LoadingPhase": "PostEngineInit"
-		},
-		{
-			"Name": "AudioModulationTest",
-			"Type": "Runtime",
-			"LoadingPhase": "PreDefault"
-		},
+		}
 	],
 	"Plugins": [
 		{
@@ -41,5 +36,5 @@
 			"Name": "WaveTable",
 			"Enabled": true
 		}
-	]	
+	]
 }
\ No newline at end of file
diff --git a/Engine/Plugins/Runtime/AudioModulation/Source/AudioModulationTest/AudioModulationTest.Build.cs b/Engine/Plugins/Runtime/AudioModulation/Source/AudioModulationTest/AudioModulationTest.Build.cs
deleted file mode 100644
diff --git a/Engine/Plugins/Runtime/AudioModulation/Source/AudioModulationTest/Private/AudioModulationTestModule.cpp b/Engine/Plugins/Runtime/AudioModulation/Source/AudioModulationTest/Private/AudioModulationTestModule.cpp
deleted file mode 100644
diff --git a/Engine/Plugins/Runtime/AudioModulation/Source/AudioModulationTest/Private/AudioModulationWatcherTest.cpp b/Engine/Plugins/Runtime/AudioModulation/Source/AudioModulationTest/Private/AudioModulationWatcherTest.cpp
deleted file mode 100644
