Additive Modulation Parameter

diff --git a/Engine/Plugins/Runtime/AudioModulation/Source/AudioModulation/Private/SoundModulationParameter.cpp b/Engine/Plugins/Runtime/AudioModulation/Source/AudioModulation/Private/SoundModulationParameter.cpp
--- a/Engine/Plugins/Runtime/AudioModulation/Source/AudioModulation/Private/SoundModulationParameter.cpp
+++ b/Engine/Plugins/Runtime/AudioModulation/Source/AudioModulation/Private/SoundModulationParameter.cpp
@@ -295,6 +295,52 @@ Audio::FModulationParameter USoundModulationParameterVolume::CreateDefaultParame
 	return NewParam;
 }
 
+USoundModulationParameterAdditive::USoundModulationParameterAdditive(const FObjectInitializer& ObjectInitializer)
+	: Super(ObjectInitializer)
+{
+	Settings.ValueNormalized = 0.0f;
+}
+
+bool USoundModulationParameterAdditive::RequiresUnitConversion() const
+{
+	return true;
+}
+
+Audio::FModulationMixFunction USoundModulationParameterAdditive::GetMixFunction() const
+{
+	return [](float& InOutValueA, float InValueB)
+	{
+		InOutValueA += InValueB;
+	};
+}
+
+Audio::FModulationUnitConversionFunction USoundModulationParameterAdditive::GetUnitConversionFunction() const
+{
+	return [InUnitMin = UnitMin, InUnitMax = UnitMax](float& InOutValue)
+	{
+		InOutValue = FMath::Lerp(InUnitMin, InUnitMax, InOutValue);
+	};
+}
+
+Audio::FModulationNormalizedConversionFunction USoundModulationParameterAdditive::GetNormalizedConversionFunction() const
+{
+	return [InUnitMin = UnitMin, InUnitMax = UnitMax](float& InOutValue)
+	{
+		const float Denom = FMath::Max(SMALL_NUMBER, InUnitMax - InUnitMin);
+		InOutValue = (InOutValue - InUnitMin) / Denom;
+	};
+}
+
+float USoundModulationParameterAdditive::GetUnitMax() const
+{
+	return UnitMax;
+}
+
+float USoundModulationParameterAdditive::GetUnitMin() const
+{
+	return UnitMin;
+}
+
 namespace AudioModulation
 {
 	const Audio::FModulationParameter& GetOrRegisterParameter(const USoundModulationParameter* InParameter, const FString& InBreadcrumb)
diff --git a/Engine/Plugins/Runtime/AudioModulation/Source/AudioModulation/Public/SoundModulationParameter.h b/Engine/Plugins/Runtime/AudioModulation/Source/AudioModulation/Public/SoundModulationParameter.h
--- a/Engine/Plugins/Runtime/AudioModulation/Source/AudioModulation/Public/SoundModulationParameter.h
+++ b/Engine/Plugins/Runtime/AudioModulation/Source/AudioModulation/Public/SoundModulationParameter.h
@@ -218,7 +218,7 @@ public:
 	static Audio::FModulationParameter CreateDefaultParameter();
 };
 
-// Modulation Parameter that scales normalized, unitless value to bipolar range. Mixes multiplicatively.
+// Modulation Parameter that scales normalized, unitless value to bipolar range. Mixes additively.
 UCLASS(BlueprintType, MinimalAPI)
 class USoundModulationParameterBipolar : public USoundModulationParameter
 {
@@ -256,6 +256,30 @@ public:
 	static Audio::FModulationParameter CreateDefaultParameter(float MinUnitVolume = -100.0f);
 };
 
+// Modulation Parameter whose values are mixed via addition.
+UCLASS(BlueprintType, MinimalAPI)
+class USoundModulationParameterAdditive : public USoundModulationParameter
+{
+	GENERATED_UCLASS_BODY()
+
+public:
+
+	/** Unit minimum of modulator. Minimum is only enforced at modulation destination. */
+	UPROPERTY(EditAnywhere, Category = General, BlueprintReadOnly)
+	float UnitMin = 0.0f;
+
+	/** Unit maximum of modulator. Maximum is only enforced at modulation destination. */
+	UPROPERTY(EditAnywhere, Category = General, BlueprintReadOnly)
+	float UnitMax = 1.0f;
+
+	virtual bool RequiresUnitConversion() const override;
+	virtual Audio::FModulationMixFunction GetMixFunction() const override;
+	virtual Audio::FModulationUnitConversionFunction GetUnitConversionFunction() const override;
+	virtual Audio::FModulationNormalizedConversionFunction GetNormalizedConversionFunction() const override;
+	virtual float GetUnitMax() const override;
+	virtual float GetUnitMin() const override;
+};
+
 namespace AudioModulation
 {
 	class FSoundModulationPluginParameterAssetProxy : public FSoundModulationParameterAssetProxy
