Clamp Unit range values to be above 0 on Additive Parameter

diff --git a/Engine/Plugins/Runtime/AudioModulation/Source/AudioModulation/Public/SoundModulationParameter.h b/Engine/Plugins/Runtime/AudioModulation/Source/AudioModulation/Public/SoundModulationParameter.h
--- a/Engine/Plugins/Runtime/AudioModulation/Source/AudioModulation/Public/SoundModulationParameter.h
+++ b/Engine/Plugins/Runtime/AudioModulation/Source/AudioModulation/Public/SoundModulationParameter.h
@@ -265,11 +265,11 @@ class USoundModulationParameterAdditive : public USoundModulationParameter
 public:
 
 	/** Unit minimum of modulator. Minimum is only enforced at modulation destination. */
-	UPROPERTY(EditAnywhere, Category = General, BlueprintReadOnly)
+	UPROPERTY(EditAnywhere, Category = General, BlueprintReadOnly, meta = (ClampMin = "0.0"))
 	float UnitMin = 0.0f;
 
 	/** Unit maximum of modulator. Maximum is only enforced at modulation destination. */
-	UPROPERTY(EditAnywhere, Category = General, BlueprintReadOnly)
+	UPROPERTY(EditAnywhere, Category = General, BlueprintReadOnly, meta = (ClampMin = "0.0"))
 	float UnitMax = 1.0f;
 
 	virtual bool RequiresUnitConversion() const override;
