Modulation debug FString/FName Optimizations

diff --git a/Engine/Plugins/Runtime/AudioModulation/Source/AudioModulation/Private/AudioModulationDebugger.cpp b/Engine/Plugins/Runtime/AudioModulation/Source/AudioModulation/Private/AudioModulationDebugger.cpp
--- a/Engine/Plugins/Runtime/AudioModulation/Source/AudioModulation/Private/AudioModulationDebugger.cpp
+++ b/Engine/Plugins/Runtime/AudioModulation/Source/AudioModulation/Private/AudioModulationDebugger.cpp
@@ -379,7 +379,7 @@ namespace AudioModulation
 			{
 				const U& Item = IdItemPair.Value;
 				const bool Filtered = !FilterString.IsEmpty()
-					&& !Item.GetName().Contains(FilterString);
+					&& !Item.GetName().ToString().Contains(FilterString);
 				if (!Filtered)
 				{
 					FilteredArray.Add(&Item);
@@ -427,7 +427,7 @@ namespace AudioModulation
 			DebugInfo.Id = Proxy->GetId();
 			DebugInfo.GeneratorValue = Proxy->GetGeneratorValue();
 			DebugInfo.MixValue = Proxy->GetMixValue();
-			DebugInfo.Name = Proxy->GetName();
+			DebugInfo.Name = Proxy->GetName().ToString();
 			DebugInfo.RefCount = Proxy->GetRefCount();
 			DebugInfo.Value = Proxy->GetValue();
 			RefreshedFilteredBuses.Add(DebugInfo);
@@ -440,7 +440,7 @@ namespace AudioModulation
 		for (const FModulatorBusMixProxy* Proxy : FilteredMixProxies)
 		{
 			FControlBusMixDebugInfo DebugInfo;
-			DebugInfo.Name = Proxy->GetName();
+			DebugInfo.Name = Proxy->GetName().ToString();
 			DebugInfo.RefCount = Proxy->GetRefCount();
 			for (const TPair<FBusId, FModulatorBusMixStageProxy>& Stage : Proxy->Stages)
 			{
diff --git a/Engine/Plugins/Runtime/AudioModulation/Source/AudioModulation/Private/AudioModulationSystem.cpp b/Engine/Plugins/Runtime/AudioModulation/Source/AudioModulation/Private/AudioModulationSystem.cpp
--- a/Engine/Plugins/Runtime/AudioModulation/Source/AudioModulation/Private/AudioModulationSystem.cpp
+++ b/Engine/Plugins/Runtime/AudioModulation/Source/AudioModulation/Private/AudioModulationSystem.cpp
@@ -448,7 +448,7 @@ namespace AudioModulation
 				{
 					if (LastStatus != CurrentStatus)
 					{
-						UE_LOG(LogAudioModulation, Verbose, TEXT("Audio modulation mix '%s' stopped."), *Pair.Value.GetName());
+						UE_LOG(LogAudioModulation, Verbose, TEXT("Audio modulation mix '%s' stopped."), *Pair.Value.GetName().ToString());
 					}
 					StoppedMixIds.Add(Pair.Key);
 				}
@@ -858,7 +858,7 @@ namespace AudioModulation
 #if !UE_BUILD_SHIPPING
 			else
 			{
-				UE_LOG(LogAudioModulation, Verbose, TEXT("Update to '%s' Ignored: Control Bus Mix is inactive."), *MixSettings.GetName());
+				UE_LOG(LogAudioModulation, Verbose, TEXT("Update to '%s' Ignored: Control Bus Mix is inactive."), *MixSettings.GetName().ToString());
 			}
 #endif // !UE_BUILD_SHIPPING
 		});
@@ -878,7 +878,7 @@ namespace AudioModulation
 #if !UE_BUILD_SHIPPING
 				else
 				{
-					UE_LOG(LogAudioModulation, Verbose, TEXT("Update to '%s' Ignored: Generator is inactive."), *GeneratorSettings.GetName());
+					UE_LOG(LogAudioModulation, Verbose, TEXT("Update to '%s' Ignored: Generator is inactive."), *GeneratorSettings.GetName().ToString());
 				}
 #endif // !UE_BUILD_SHIPPING
 			});
@@ -897,7 +897,7 @@ namespace AudioModulation
 #if !UE_BUILD_SHIPPING
 				else
 				{
-					UE_LOG(LogAudioModulation, Verbose, TEXT("Update to '%s' Ignored: Control Bus is inactive."), *BusSettings.GetName());
+					UE_LOG(LogAudioModulation, Verbose, TEXT("Update to '%s' Ignored: Control Bus is inactive."), *BusSettings.GetName().ToString());
 				}
 #endif // !UE_BUILD_SHIPPING
 			});
@@ -916,7 +916,7 @@ namespace AudioModulation
 #if !UE_BUILD_SHIPPING
 				else
 				{
-					UE_LOG(LogAudioModulation, Verbose, TEXT("Update to '%s' Ignored: Patch is inactive."), *PatchSettings.GetName());
+					UE_LOG(LogAudioModulation, Verbose, TEXT("Update to '%s' Ignored: Patch is inactive."), *PatchSettings.GetName().ToString());
 				}
 #endif // !UE_BUILD_SHIPPING
 			});
diff --git a/Engine/Plugins/Runtime/AudioModulation/Source/AudioModulation/Private/SoundControlBus.cpp b/Engine/Plugins/Runtime/AudioModulation/Source/AudioModulation/Private/SoundControlBus.cpp
--- a/Engine/Plugins/Runtime/AudioModulation/Source/AudioModulation/Private/SoundControlBus.cpp
+++ b/Engine/Plugins/Runtime/AudioModulation/Source/AudioModulation/Private/SoundControlBus.cpp
@@ -149,11 +149,9 @@ TSharedPtr<Audio::IProxyData> USoundControlBus::CreateProxyData(const Audio::FPr
 const Audio::FModulationParameter& USoundControlBus::GetOutputParameter() const
 {
 #if !UE_BUILD_SHIPPING
-	TStringBuilder<128> StringBuilder;
-	StringBuilder.Append(*GetClass()->GetName()).Append(" '").Append(*GetName()).Append("'");
-	return AudioModulation::GetOrRegisterParameter(Parameter, StringBuilder.ToString());
+	return AudioModulation::GetOrRegisterParameter(Parameter, GetName(), GetClass()->GetName());
 #else
-	return AudioModulation::GetOrRegisterParameter(Parameter, FString());
+	return AudioModulation::GetOrRegisterParameter(Parameter, FString(), FString());
 #endif
 }
 
diff --git a/Engine/Plugins/Runtime/AudioModulation/Source/AudioModulation/Private/SoundModulationGeneratorProxy.h b/Engine/Plugins/Runtime/AudioModulation/Source/AudioModulation/Private/SoundModulationGeneratorProxy.h
--- a/Engine/Plugins/Runtime/AudioModulation/Source/AudioModulation/Private/SoundModulationGeneratorProxy.h
+++ b/Engine/Plugins/Runtime/AudioModulation/Source/AudioModulation/Private/SoundModulationGeneratorProxy.h
@@ -114,7 +114,7 @@ namespace AudioModulation
 		TArray<FString> GetDebugValues() const
 		{
 			TArray<FString> DebugValues;
-			DebugValues.Add(GetName());
+			DebugValues.Add(GetName().ToString());
 			DebugValues.Add(FString::FormatAsNumber(GetRefCount()));
 
 			Generator->GetDebugValues(DebugValues);
diff --git a/Engine/Plugins/Runtime/AudioModulation/Source/AudioModulation/Private/SoundModulationParameter.cpp b/Engine/Plugins/Runtime/AudioModulation/Source/AudioModulation/Private/SoundModulationParameter.cpp
--- a/Engine/Plugins/Runtime/AudioModulation/Source/AudioModulation/Private/SoundModulationParameter.cpp
+++ b/Engine/Plugins/Runtime/AudioModulation/Source/AudioModulation/Private/SoundModulationParameter.cpp
@@ -343,7 +343,7 @@ float USoundModulationParameterAdditive::GetUnitMin() const
 
 namespace AudioModulation
 {
-	const Audio::FModulationParameter& GetOrRegisterParameter(const USoundModulationParameter* InParameter, const FString& InBreadcrumb)
+	const Audio::FModulationParameter& GetOrRegisterParameter(const USoundModulationParameter* InParameter, const FString& InName, const FString& InClassName)
 	{
 		FName ParamName;
 		if (InParameter)
@@ -351,10 +351,20 @@ namespace AudioModulation
 			ParamName = InParameter->GetFName();
 			if (!Audio::IsModulationParameterRegistered(ParamName))
 			{
+				TStringBuilder<128> Breadcrumb;
+				if (InClassName.IsEmpty())
+				{
+					Breadcrumb.Append(*InName);
+				}
+				else
+				{
+					Breadcrumb.Append(*InClassName).Append(" '").Append(*InName).Append("'");
+				}
+
 				UE_LOG(LogAudioModulation, Display,
 					TEXT("Parameter '%s' not registered.  Registration forced via '%s'."),
 					*ParamName.ToString(),
-					*InBreadcrumb);
+					*Breadcrumb);
 
 				Audio::RegisterModulationParameter(ParamName, InParameter->CreateParameter());
 			}
@@ -366,6 +376,6 @@ namespace AudioModulation
 
 	FSoundModulationPluginParameterAssetProxy::FSoundModulationPluginParameterAssetProxy(USoundModulationParameter* InParameter)
 	{
-		Parameter = GetOrRegisterParameter(InParameter, TEXT("FSoundModulationPluginParameterAssetProxy construction"));
+		Parameter = GetOrRegisterParameter(InParameter, TEXT("FSoundModulationPluginParameterAssetProxy construction"), FString());
 	}
 }
diff --git a/Engine/Plugins/Runtime/AudioModulation/Source/AudioModulation/Private/SoundModulationPatch.cpp b/Engine/Plugins/Runtime/AudioModulation/Source/AudioModulation/Private/SoundModulationPatch.cpp
--- a/Engine/Plugins/Runtime/AudioModulation/Source/AudioModulation/Private/SoundModulationPatch.cpp
+++ b/Engine/Plugins/Runtime/AudioModulation/Source/AudioModulation/Private/SoundModulationPatch.cpp
@@ -35,8 +35,7 @@ TUniquePtr<Audio::IModulatorSettings> USoundModulationPatch::CreateProxySettings
 
 const Audio::FModulationParameter& USoundModulationPatch::GetOutputParameter() const
 {
-	const FString Breadcrumb = FString::Format(TEXT("{0} '{1}'"), { *GetClass()->GetName(), *GetName() });
-	return AudioModulation::GetOrRegisterParameter(PatchSettings.OutputParameter, Breadcrumb);
+	return AudioModulation::GetOrRegisterParameter(PatchSettings.OutputParameter, GetName(), GetClass()->GetName());
 }
 
 #if WITH_EDITOR
diff --git a/Engine/Plugins/Runtime/AudioModulation/Source/AudioModulation/Private/SoundModulationProxy.h b/Engine/Plugins/Runtime/AudioModulation/Source/AudioModulation/Private/SoundModulationProxy.h
--- a/Engine/Plugins/Runtime/AudioModulation/Source/AudioModulation/Private/SoundModulationProxy.h
+++ b/Engine/Plugins/Runtime/AudioModulation/Source/AudioModulation/Private/SoundModulationProxy.h
@@ -103,7 +103,7 @@ namespace AudioModulation
 #if UE_BUILD_SHIPPING
 				UE_LOG(LogAudioModulation, Verbose, TEXT("Proxy released: Id '%u'."), Id);
 #else // UE_BUILD_SHIPPING
-				UE_LOG(LogAudioModulation, Verbose, TEXT("Proxy '%s' released: Id '%u'."), *Proxy.GetName(), Id);
+				UE_LOG(LogAudioModulation, Verbose, TEXT("Proxy '%s' released: Id '%u'."), *Proxy.GetName().ToString(), Id);
 #endif // !UE_BUILD_SHIPPING
 				ProxyMap->Remove(Id);
 			}
@@ -133,7 +133,7 @@ namespace AudioModulation
 			TProxyHandle<IdType, ProxyType, ProxySettings> NewHandle = Get(InSettings.GetId(), OutProxyMap);
 			if (!NewHandle.IsValid())
 			{
-				UE_LOG(LogAudioModulation, Verbose, TEXT("Proxy created: Id '%u' for object '%s'."), ObjectId, *InSettings.GetName());
+				UE_LOG(LogAudioModulation, Verbose, TEXT("Proxy created: Id '%u' for object '%s'."), ObjectId, *InSettings.GetName().ToString());
 				OutProxyMap.Add(ObjectId, ProxyType(MoveTemp(InSettings), OutModSystem));
 				NewHandle = TProxyHandle<IdType, ProxyType, ProxySettings>(ObjectId, OutProxyMap);
 			}
@@ -237,7 +237,7 @@ namespace AudioModulation
 		IdType Id;
 
 #if !UE_BUILD_SHIPPING
-		FString Name;
+		FName Name;
 #endif // !UE_BUILD_SHIPPING
 
 	public:
@@ -249,7 +249,7 @@ namespace AudioModulation
 		TModulatorBase(const FName InName, const uint32 InId)
 			: Id(static_cast<IdType>(InId))
 #if !UE_BUILD_SHIPPING
-			, Name(InName.ToString())
+			, Name(InName)
 #endif // !UE_BUILD_SHIPPING
 		{
 		}
@@ -271,10 +271,10 @@ namespace AudioModulation
 
 		// FOR DEBUG USE ONLY (Not available in shipped builds):
 		// Provides name of object that generated proxy.
-		const FString& GetName() const
+		const FName& GetName() const
 		{
 #if UE_BUILD_SHIPPING
-			static FString Name;
+			static FName Name;
 #endif // UE_BUILD_SHIPPING
 
 			return Name;
@@ -296,7 +296,7 @@ namespace AudioModulation
 		{
 		}
 
-		TModulatorProxyRefType(const FString& InName, const uint32 InId, FAudioModulationSystem& OutModSystem)
+		TModulatorProxyRefType(const FName& InName, const uint32 InId, FAudioModulationSystem& OutModSystem)
 			: TModulatorBase<IdType>(InName, InId)
 			, RefCount(0)
 			, ModSystem(&OutModSystem)
diff --git a/Engine/Plugins/Runtime/AudioModulation/Source/AudioModulation/Public/SoundModulationParameter.h b/Engine/Plugins/Runtime/AudioModulation/Source/AudioModulation/Public/SoundModulationParameter.h
--- a/Engine/Plugins/Runtime/AudioModulation/Source/AudioModulation/Public/SoundModulationParameter.h
+++ b/Engine/Plugins/Runtime/AudioModulation/Source/AudioModulation/Public/SoundModulationParameter.h
@@ -292,7 +292,9 @@ namespace AudioModulation
 	/*
 	 * Returns given registered parameter instance reference or creates it from the given asset if not registered.
 	 * @param InParameter - Parameter asset associated with the pre-existing or to-create parameter
-	 * @param InBreadcrumb - String identifying get or register request initiator.
+	 * @param InName	  - Name of the modulator requesting the parameter.
+	 * @param InClassName - Name of the modulator class requesting the parameter.
+	 * 
 	 */
-	AUDIOMODULATION_API const Audio::FModulationParameter& GetOrRegisterParameter(const USoundModulationParameter* InParameter, const FString& InBreadcrumb);
+	AUDIOMODULATION_API const Audio::FModulationParameter& GetOrRegisterParameter(const USoundModulationParameter* InParameter, const FString& InName, const FString& InClassName);
 } // namespace AudioModulation
