Fix missing null implementation calls for AudioModulationSystem
#rnx

diff --git a/Engine/Plugins/Runtime/AudioModulation/Source/AudioModulation/Private/AudioModulationSystem.h b/Engine/Plugins/Runtime/AudioModulation/Source/AudioModulation/Private/AudioModulationSystem.h
--- a/Engine/Plugins/Runtime/AudioModulation/Source/AudioModulation/Private/AudioModulationSystem.h
+++ b/Engine/Plugins/Runtime/AudioModulation/Source/AudioModulation/Private/AudioModulationSystem.h
@@ -279,10 +279,6 @@ namespace AudioModulation
 	public:
 		void Initialize(const FAudioPluginInitializationParams& InitializationParams) { }
 
-#if WITH_EDITOR
-		void SoloBusMix(const USoundControlBusMix& InBusMix) { }
-#endif // WITH_EDITOR
-
 		void OnAuditionEnd() { }
 
 #if !UE_BUILD_SHIPPING
@@ -310,13 +306,14 @@ namespace AudioModulation
 		Audio::FDeviceId GetAudioDeviceId() const { return 0; }
 
 		void SaveMixToProfile(const USoundControlBusMix& InBusMix, const int32 InProfileIndex) { }
-		TArray<FSoundControlBusMixStage> LoadMixFromProfile(const int32 InProfileIndex, USoundControlBusMix& OutBusMix) { }
+		TArray<FSoundControlBusMixStage> LoadMixFromProfile(const int32 InProfileIndex, USoundControlBusMix& OutBusMix) { return { }; }
 
 		void SetGlobalBusMixValue(USoundControlBus& Bus, float Value, float FadeTime) { }
 		void ClearGlobalBusMixValue(const USoundControlBus& InBus, float FadeTime) { }
 		void ClearAllGlobalBusMixValues(float FadeTime) { }
 
 		void ProcessModulators(const double InElapsed) { }
+		void SoloBusMix(const USoundControlBusMix& InBusMix) { }
 
 		Audio::FModulatorTypeId RegisterModulator(Audio::FModulatorHandleId InHandleId, const FControlBusSettings& InSettings) { return 0; }
 		Audio::FModulatorTypeId RegisterModulator(Audio::FModulatorHandleId InHandleId, const FModulationGeneratorSettings& InSettings) { return 0; }
@@ -325,6 +322,7 @@ namespace AudioModulation
 		void RegisterModulator(Audio::FModulatorHandleId InHandleId, Audio::FModulatorId InModulatorId) { }
 		bool GetModulatorValue(const Audio::FModulatorHandle& ModulatorHandle, float& OutValue) const { return false; }
 		bool GetModulatorValueThreadSafe(const Audio::FModulatorHandle& ModulatorHandle, float& OutValue) const { return false; }
+		bool GetModulatorValueThreadSafe(uint32 ModulatorID, float& OutValue) const { return false; }
 		void UnregisterModulator(const Audio::FModulatorHandle& InHandle) { }
 
 		void UpdateMix(const USoundControlBusMix& InMix, float InFadeTime = -1.0f) { }
