Fix crash in SetGlobalControlBusMixValue on PIE exit

diff --git a/Engine/Plugins/Runtime/AudioModulation/Source/AudioModulation/Private/AudioModulationSystem.cpp b/Engine/Plugins/Runtime/AudioModulation/Source/AudioModulation/Private/AudioModulationSystem.cpp
--- a/Engine/Plugins/Runtime/AudioModulation/Source/AudioModulation/Private/AudioModulationSystem.cpp
+++ b/Engine/Plugins/Runtime/AudioModulation/Source/AudioModulation/Private/AudioModulationSystem.cpp
@@ -242,28 +242,33 @@ namespace AudioModulation
 		else
 		{
 			const FString MixName = InBus.GetName() + TEXT("_GlobalMix");
-			TObjectPtr<USoundControlBusMix> NewGlobalMix = NewObject<USoundControlBusMix>(GetTransientPackage(), FName(*MixName));
-
+			if (TObjectPtr<USoundControlBusMix> NewGlobalMix = NewObject<USoundControlBusMix>(GetTransientPackage(), FName(*MixName)))
 			{
-				FSoundModulationMixValue MixValue;
-				MixValue.TargetValue = InValue;
-
-				if (InFadeTime >= 0.0f)
 				{
-					MixValue.AttackTime = InFadeTime;
-				}
+					FSoundModulationMixValue MixValue;
+					MixValue.TargetValue = InValue;
 
-				FSoundControlBusMixStage MixStage;
-				MixStage.Bus = &InBus;
-				MixStage.Value = MixValue;
+					if (InFadeTime >= 0.0f)
+					{
+						MixValue.AttackTime = InFadeTime;
+					}
 
-				NewGlobalMix->MixStages.Emplace(MoveTemp(MixStage));
-			}
+					FSoundControlBusMixStage MixStage;
+					MixStage.Bus = &InBus;
+					MixStage.Value = MixValue;
 
-			ActiveGlobalBusValueMixes.Add(InBus.GetUniqueID(), NewGlobalMix);
-			UE_LOG(LogAudioModulation, VeryVerbose, TEXT("GlobalBusMix for ControlBus '%s' activated, target set to '%0.4f'."), *InBus.GetName(), InValue);
-			ActivateBusMix(*NewGlobalMix);
-			NewGlobalMix->AddToRoot();
+					NewGlobalMix->MixStages.Emplace(MoveTemp(MixStage));
+				}
+
+				ActiveGlobalBusValueMixes.Add(InBus.GetUniqueID(), NewGlobalMix);
+				UE_LOG(LogAudioModulation, VeryVerbose, TEXT("GlobalBusMix for ControlBus '%s' activated, target set to '%0.4f'."), *InBus.GetName(), InValue);
+				ActivateBusMix(*NewGlobalMix);
+				NewGlobalMix->AddToRoot();
+			}
+			else
+			{
+				UE_LOG(LogAudioModulation, Warning, TEXT("Could not set Global Bus Mix value: failed to allocate new Global Bus Mix for bus %s."), *InBus.GetName());
+			}
 		}
 
 	}
