Remove Unnecessary warning when activating a control bus

diff --git a/Engine/Plugins/Runtime/AudioModulation/Source/AudioModulation/Private/AudioModulationSystem.cpp b/Engine/Plugins/Runtime/AudioModulation/Source/AudioModulation/Private/AudioModulationSystem.cpp
--- a/Engine/Plugins/Runtime/AudioModulation/Source/AudioModulation/Private/AudioModulationSystem.cpp
+++ b/Engine/Plugins/Runtime/AudioModulation/Source/AudioModulation/Private/AudioModulationSystem.cpp
@@ -771,10 +771,6 @@ namespace AudioModulation
 			{
 				BusMixes->SetMix(StageSettings, InFadeTime, BusMixName);
 			}
-			else
-			{
-				UE_LOG(LogAudioModulation, Warning, TEXT("Updating Control Bus Mix failed: Bus Mix '%s' not currently active. Buses without any valid Mix Stages will not activate."), *BusMixName);
-			}
 		});
 	}
 
