Modulation Swins

diff --git a/Engine/Plugins/Runtime/AudioModulation/Source/AudioModulation/Public/SoundControlBus.h b/Engine/Plugins/Runtime/AudioModulation/Source/AudioModulation/Public/SoundControlBus.h
--- a/Engine/Plugins/Runtime/AudioModulation/Source/AudioModulation/Public/SoundControlBus.h
+++ b/Engine/Plugins/Runtime/AudioModulation/Source/AudioModulation/Public/SoundControlBus.h
@@ -32,17 +32,19 @@ public:
 
 #if WITH_EDITORONLY_DATA
 	/** If true, Address field is used in place of object name for address used when applying mix changes using filtering. */
-	UPROPERTY(EditAnywhere, Category = Mix, BlueprintReadWrite)
+	UPROPERTY(EditAnywhere, Category = "Advanced|Mix", BlueprintReadWrite)
 	bool bOverrideAddress;
 #endif // WITH_EDITORONLY_DATA
 
 	/** Address to use when applying mix changes. */
-	UPROPERTY(EditAnywhere, Category = Mix, BlueprintReadWrite, meta = (EditCondition = "bOverrideAddress"))
+	UPROPERTY(EditAnywhere, Category = "Advanced|Mix", BlueprintReadWrite, meta = (EditCondition = "bOverrideAddress"))
 	FString Address;
 
-	UPROPERTY(EditAnywhere, Category = Generators, BlueprintReadWrite)
+	/** The Modulation Generators which modulate this Control Bus. */
+	UPROPERTY(EditAnywhere, Category = "Advanced|Generators", BlueprintReadWrite)
 	TArray<TObjectPtr<USoundModulationGenerator>> Generators;
 
+	/** The Modulation Parameter to use with this Control Bus. This determines how the control bus will be mixed with other modulators as well as its unit range. */
 	UPROPERTY(EditAnywhere, Category = General, BlueprintReadOnly)
 	TObjectPtr<USoundModulationParameter> Parameter;
 
diff --git a/Engine/Plugins/Runtime/AudioModulation/Source/AudioModulation/Public/SoundModulationPatch.h b/Engine/Plugins/Runtime/AudioModulation/Source/AudioModulation/Public/SoundModulationPatch.h
--- a/Engine/Plugins/Runtime/AudioModulation/Source/AudioModulation/Public/SoundModulationPatch.h
+++ b/Engine/Plugins/Runtime/AudioModulation/Source/AudioModulation/Public/SoundModulationPatch.h
@@ -31,15 +31,15 @@ struct AUDIOMODULATION_API FSoundControlModulationInput
 	FSoundControlModulationInput();
 
 	/** Get the modulated input value on parent patch initialization and hold that value for its lifetime */
-	UPROPERTY(EditAnywhere, BlueprintReadWrite, Category = Input, meta = (DisplayName = "Sample-And-Hold"))
+	UPROPERTY(EditAnywhere, BlueprintReadWrite, Category = Input, meta = (DisplayName = "Sample-And-Hold", DisplayPriority = "30"))
 	uint8 bSampleAndHold : 1;
 
 	/** Transform to apply to the input prior to mix phase */
-	UPROPERTY(EditAnywhere, Category = Input)
+	UPROPERTY(EditAnywhere, Category = Input, meta = (DisplayPriority = "20"))
 	FSoundModulationTransform Transform;
 
 	/** The input bus */
-	UPROPERTY(EditAnywhere, BlueprintReadWrite, Category = Input)
+	UPROPERTY(EditAnywhere, BlueprintReadWrite, Category = Input, meta = (DisplayPriority = "10"))
 	TObjectPtr<USoundControlBus> Bus = nullptr;
 
 	const USoundControlBus* GetBus() const;
diff --git a/Engine/Plugins/Runtime/WaveTable/Source/WaveTable/Public/WaveTableSettings.h b/Engine/Plugins/Runtime/WaveTable/Source/WaveTable/Public/WaveTableSettings.h
--- a/Engine/Plugins/Runtime/WaveTable/Source/WaveTable/Public/WaveTableSettings.h
+++ b/Engine/Plugins/Runtime/WaveTable/Source/WaveTable/Public/WaveTableSettings.h
@@ -164,7 +164,7 @@ enum class EWaveTableCurve : uint8
 	SCurve		UMETA(DisplayName = "Sin (+/- 90 deg)"),
 
 	// Reference a shared curve asset
-	Shared		UMETA(DisplayName = "Shared"),
+	Shared		UMETA(DisplayName = "Curve Asset"),
 
 	// Design a custom curve unique to the owning transform
 	Custom		UMETA(DisplayName = "Custom"),
