Rename SoundModulationWatcher to AudioModulationDestination
#tests created old type, added to BP, then added redirect and verified old behavior still worked.

diff --git a/Engine/Plugins/Runtime/AudioModulation/Config/BaseAudioModulation.ini b/Engine/Plugins/Runtime/AudioModulation/Config/BaseAudioModulation.ini
--- a/Engine/Plugins/Runtime/AudioModulation/Config/BaseAudioModulation.ini
+++ b/Engine/Plugins/Runtime/AudioModulation/Config/BaseAudioModulation.ini
@@ -1,6 +1,7 @@
 [CoreRedirects]
 +EnumRedirects=(OldName="FSoundModulationLFOShape",NewName="/Script/AudioModulation.ESoundModulationLFOShape")
 +EnumRedirects=(OldName="ESoundModulationCurve",NewName="/Script/WaveTable.EWaveTableCurve")
++ClassRedirects=(OldName="SoundModulationWatcher",NewName="/Script/AudioModulation.SoundModulationDestination")
 
 [/Script/AudioModulation.AudioModulationSettings]
 +Parameters=/AudioModulation/BitDepth.BitDepth
diff --git a/Engine/Plugins/Runtime/AudioModulation/Source/AudioModulation/Private/AudioModulation.cpp b/Engine/Plugins/Runtime/AudioModulation/Source/AudioModulation/Private/AudioModulation.cpp
--- a/Engine/Plugins/Runtime/AudioModulation/Source/AudioModulation/Private/AudioModulation.cpp
+++ b/Engine/Plugins/Runtime/AudioModulation/Source/AudioModulation/Private/AudioModulation.cpp
@@ -48,7 +48,7 @@ namespace AudioModulation
 
 	void FAudioModulationManager::ActivateBus(const USoundControlBus& InBus)
 	{
-		UE_LOG(LogAudioModulation, Warning, TEXT("FAudioModulationManager::ActivateBus is deprecated. Use USoundModulationWatcher instead"));
+		UE_LOG(LogAudioModulation, Warning, TEXT("FAudioModulationManager::ActivateBus is deprecated. Use UAudioModulationDestination API (see SetModulator) instead to enforce generator lifetime"));
 
 		PRAGMA_DISABLE_DEPRECATION_WARNINGS
 		ModSystem->ActivateBus(InBus);
@@ -62,7 +62,7 @@ namespace AudioModulation
 
 	void FAudioModulationManager::ActivateGenerator(const USoundModulationGenerator& InGenerator)
 	{
-		UE_LOG(LogAudioModulation, Warning, TEXT("FAudioModulationManager::ActivateGenerator is deprecated. Use USoundModulationWatcher instead"));
+		UE_LOG(LogAudioModulation, Warning, TEXT("FAudioModulationManager::ActivateGenerator is deprecated. Use UAudioModulationDestination API (see SetModulator) instead to enforce generator lifetime"));
 
 		PRAGMA_DISABLE_DEPRECATION_WARNINGS
 		ModSystem->ActivateGenerator(InGenerator);
@@ -71,7 +71,7 @@ namespace AudioModulation
 
 	void FAudioModulationManager::DeactivateBus(const USoundControlBus& InBus)
 	{
-		UE_LOG(LogAudioModulation, Warning, TEXT("FAudioModulationManager::DeactivateBus is deprecated. Use USoundModulationWatcher instead"));
+		UE_LOG(LogAudioModulation, Warning, TEXT("FAudioModulationManager::DeactivateBus is deprecated. Use UAudioModulationDestination API (see ClearModulator) instead to enforce bus lifetime."));
 
 		PRAGMA_DISABLE_DEPRECATION_WARNINGS
 		ModSystem->DeactivateBus(InBus);
@@ -90,7 +90,7 @@ namespace AudioModulation
 
 	void FAudioModulationManager::DeactivateGenerator(const USoundModulationGenerator& InGenerator)
 	{
-		UE_LOG(LogAudioModulation, Warning, TEXT("FAudioModulationManager::DeactivateGenerator is deprecated. Use USoundModulationWatcher instead"));
+		UE_LOG(LogAudioModulation, Warning, TEXT("FAudioModulationManager::DeactivateGenerator is deprecated. Use UAudioModulationDestination API (see ClearModulator) instead to enforce bus lifetime"));
 
 		PRAGMA_DISABLE_DEPRECATION_WARNINGS
 		ModSystem->DeactivateGenerator(InGenerator);
diff --git a/Engine/Plugins/Runtime/AudioModulation/Source/AudioModulation/Private/AudioModulationDestination.cpp b/Engine/Plugins/Runtime/AudioModulation/Source/AudioModulation/Private/AudioModulationDestination.cpp
new file mode 100644
--- /dev/null
+++ b/Engine/Plugins/Runtime/AudioModulation/Source/AudioModulation/Private/AudioModulationDestination.cpp
@@ -0,0 +1,117 @@
+// Copyright Epic Games, Inc. All Rights Reserved.
+#include "AudioModulationDestination.h"
+
+#include "AudioDeviceHandle.h"
+#include "AudioModulation.h"
+#include "Engine/Engine.h"
+#include "Engine/EngineBaseTypes.h"
+#include "Engine/World.h"
+
+#include UE_INLINE_GENERATED_CPP_BY_NAME(AudioModulationDestination)
+
+
+namespace AudioModulation::DestinationPrivate
+{
+	FAudioDeviceHandle GetAudioDevice(const UAudioModulationDestination& InDestination)
+	{
+		if (GEngine)
+		{
+			if (GEngine->UseSound())
+			{
+				UWorld* World = GEngine->GetWorldFromContextObject(&InDestination, EGetWorldErrorMode::ReturnNull);
+				if (World && World->bAllowAudioPlayback && !World->IsNetMode(NM_DedicatedServer))
+				{
+					return World->GetAudioDevice();
+				}
+				else
+				{
+					return GEngine->GetMainAudioDevice();
+				}
+			}
+		}
+
+		return { };
+	}
+
+	FAudioModulationManager* GetModulationManager(const UAudioModulationDestination& InDestination)
+	{
+		FAudioDeviceHandle AudioDevice = DestinationPrivate::GetAudioDevice(InDestination);
+		if (AudioDevice.IsValid() && AudioDevice->IsModulationPluginEnabled())
+		{
+			if (IAudioModulationManager* ModulationInterface = AudioDevice->ModulationInterface.Get())
+			{
+				return static_cast<FAudioModulationManager*>(ModulationInterface);
+			}
+		}
+
+		return nullptr;
+	}
+} // AudioModulation::DestinationPrivate
+
+bool UAudioModulationDestination::ClearModulator()
+{
+	if (Modulator)
+	{
+		Modulator = nullptr;
+		Destination.UpdateModulators(TSet<const USoundModulatorBase*>{ });
+		return true;
+	}
+
+	return false;
+}
+
+const USoundModulatorBase* UAudioModulationDestination::GetModulator() const
+{
+	return Modulator;
+}
+
+float UAudioModulationDestination::GetValue() const
+{
+	using namespace AudioModulation;
+
+	if (Modulator)
+	{
+		if (FAudioModulationManager* Modulation = DestinationPrivate::GetModulationManager(*this))
+		{
+			return Modulation->GetModulatorValueThreadSafe(Modulator->GetUniqueID());
+		}
+	}
+
+	return 1.0f;
+}
+
+void UAudioModulationDestination::PostInitProperties()
+{
+	using namespace AudioModulation;
+
+	Super::PostInitProperties();
+
+	if (UAudioModulationDestination::StaticClass()->GetDefaultObject() == this)
+	{
+		return;
+	}
+
+	FAudioDeviceHandle AudioDevice = DestinationPrivate::GetAudioDevice(*this);
+	if (AudioDevice.IsValid())
+	{
+		constexpr bool bIsBuffered = false;
+		Destination.Init(AudioDevice.GetDeviceID(), bIsBuffered);
+	}
+}
+
+bool UAudioModulationDestination::SetModulator(const USoundModulatorBase* InModulator)
+{
+	if (InModulator == Modulator)
+	{
+		return true;
+	}
+
+	if (InModulator)
+	{
+		Modulator = InModulator;
+		Destination.UpdateModulators(TSet<const USoundModulatorBase*>{ InModulator });
+		return true;
+	}
+
+	return ClearModulator();
+}
diff --git a/Engine/Plugins/Runtime/AudioModulation/Source/AudioModulation/Private/AudioModulationStatics.cpp b/Engine/Plugins/Runtime/AudioModulation/Source/AudioModulation/Private/AudioModulationStatics.cpp
--- a/Engine/Plugins/Runtime/AudioModulation/Source/AudioModulation/Private/AudioModulationStatics.cpp
+++ b/Engine/Plugins/Runtime/AudioModulation/Source/AudioModulation/Private/AudioModulationStatics.cpp
@@ -4,6 +4,7 @@
 #include "Async/Async.h"
 #include "AudioDevice.h"
 #include "AudioModulation.h"
+#include "AudioModulationDestination.h"
 #include "AudioModulationLogging.h"
 #include "AudioModulationProfileSerializer.h"
 #include "AudioModulationSystem.h"
@@ -15,7 +16,6 @@
 #include "Misc/Paths.h"
 #include "SoundControlBus.h"
 #include "SoundControlBusMix.h"
-#include "SoundModulationWatcher.h"
 
 #include UE_INLINE_GENERATED_CPP_BY_NAME(AudioModulationStatics)
 
@@ -335,17 +335,22 @@ USoundModulationGeneratorADEnvelope* UAudioModulationStatics::CreateADEnvelopeGe
 	return NewGenerator;
 }
 
-USoundModulationWatcher* UAudioModulationStatics::CreateModulationWatcher(UObject* WorldContextObject, FName Name, USoundModulatorBase* Modulator)
+UAudioModulationDestination* UAudioModulationStatics::CreateModulationDestination(UObject* WorldContextObject, FName Name, USoundModulatorBase* Modulator)
 {
-	if (USoundModulationWatcher* NewWatcher = NewObject<USoundModulationWatcher>(WorldContextObject, Name))
+	if (UAudioModulationDestination* NewDestination = NewObject<UAudioModulationDestination>(WorldContextObject, Name))
 	{
-		NewWatcher->SetModulator(Modulator);
-		return NewWatcher;
+		NewDestination->SetModulator(Modulator);
+		return NewDestination;
 	}
 
 	return nullptr;
 }
 
+UAudioModulationDestination* UAudioModulationStatics::CreateModulationWatcher(UObject* WorldContextObject, FName Name, USoundModulatorBase* Modulator)
+{
+	return CreateModulationDestination(WorldContextObject, Name, Modulator);
+}
+
 void UAudioModulationStatics::DeactivateBus(const UObject* WorldContextObject, USoundControlBus* Bus)
 {
 	if (Bus)
diff --git a/Engine/Plugins/Runtime/AudioModulation/Source/AudioModulation/Private/SoundModulationWatcher.cpp b/Engine/Plugins/Runtime/AudioModulation/Source/AudioModulation/Private/SoundModulationWatcher.cpp
deleted file mode 100644
diff --git a/Engine/Plugins/Runtime/AudioModulation/Source/AudioModulation/Public/AudioModulationDestination.h b/Engine/Plugins/Runtime/AudioModulation/Source/AudioModulation/Public/AudioModulationDestination.h
new file mode 100644
--- /dev/null
+++ b/Engine/Plugins/Runtime/AudioModulation/Source/AudioModulation/Public/AudioModulationDestination.h
@@ -0,0 +1,41 @@
+// Copyright Epic Games, Inc. All Rights Reserved.
+#pragma once
+
+#include "IAudioModulation.h"
+
+#include "AudioThread.h"
+#include "Sound/SoundModulationDestination.h"
+
+#include "AudioModulationDestination.generated.h"
+
+
+UCLASS(config = Engine, editinlinenew, BlueprintType)
+class AUDIOMODULATION_API UAudioModulationDestination : public UObject
+{
+	GENERATED_BODY()
+
+private:
+	UPROPERTY()
+	TObjectPtr<const USoundModulatorBase> Modulator;
+
+	Audio::FModulationDestination Destination;
+
+public:
+	virtual void PostInitProperties() override;
+
+	// Returns true if a modulator was set and has been cleared.
+	UFUNCTION(BlueprintCallable, Category = "Audio|Modulation", DisplayName = "Clear Modulator")
+	UPARAM(DisplayName = "Modulator Cleared") bool ClearModulator();
+
+	// Returns currently set modulator.
+	UFUNCTION(BlueprintPure, Category = "Audio|Modulation")
+	UPARAM(DisplayName = "Modulator") const USoundModulatorBase* GetModulator() const;
+
+	// Returns the last calculated modulator value sampled from the thread controls are processed on.
+	UFUNCTION(BlueprintPure, Category = "Audio|Modulation", DisplayName = "Get Watched Modulator Value")
+	UPARAM(DisplayName = "Value") float GetValue() const;
+
+	// Returns true if modulator was set to new value or was already set to provided value.
+	UFUNCTION(BlueprintCallable, Category = "Audio|Modulation", DisplayName = "Set Watched Modulator")
+	UPARAM(DisplayName = "Modulator Set") bool SetModulator(const USoundModulatorBase* InModulator);
+};
diff --git a/Engine/Plugins/Runtime/AudioModulation/Source/AudioModulation/Public/AudioModulationStatics.h b/Engine/Plugins/Runtime/AudioModulation/Source/AudioModulation/Public/AudioModulationStatics.h
--- a/Engine/Plugins/Runtime/AudioModulation/Source/AudioModulation/Public/AudioModulationStatics.h
+++ b/Engine/Plugins/Runtime/AudioModulation/Source/AudioModulation/Public/AudioModulationStatics.h
@@ -44,8 +44,8 @@ public:
 	 * if bus is manually deactivated or destroyed (i.e. will not deactivate
 	 * when all references become inactive).
 	 */
-	UFUNCTION(BlueprintCallable, Category = "Audio|Modulation", DisplayName = "Activate Control Bus", meta = (
-		WorldContext = "WorldContextObject", 
+	UFUNCTION(BlueprintCallable, Category = "Audio|Modulation", DisplayName = "Activate Control Bus (Deprecated - 5.4)", meta = (
+		WorldContext = "WorldContextObject",
 		Keywords = "modulation modulator")
 	)
 	static void ActivateBus(const UObject* WorldContextObject, USoundControlBus* Bus);
@@ -56,7 +56,7 @@ public:
 	 * @param BusMix - Mix to activate
 	 */
 	UFUNCTION(BlueprintCallable, Category = "Audio|Modulation", DisplayName = "Activate Control Bus Mix", meta = (
-		WorldContext = "WorldContextObject", 
+		WorldContext = "WorldContextObject",
 		Keywords = "modulation modulator")
 	)
 	static void ActivateBusMix(const UObject* WorldContextObject, USoundControlBusMix* Mix);
@@ -67,8 +67,8 @@ public:
 	 * when all references become inactive).
 	 * @param Modulator - Modulator to activate
 	 */
-	UFUNCTION(BlueprintCallable, Category = "Audio|Modulation", DisplayName = "Activate Modulation Generator", meta = (
-		WorldContext = "WorldContextObject", 
+	UFUNCTION(BlueprintCallable, Category = "Audio|Modulation", DisplayName = "Activate Modulation Generator (Deprecated - 5.4)", meta = (
+		WorldContext = "WorldContextObject",
 		Keywords = "modulator lfo envelope follower")
 	)
 	static void ActivateGenerator(const UObject* WorldContextObject, USoundModulationGenerator* Generator);
@@ -76,10 +76,10 @@ public:
 	/** Creates a modulation bus with the provided default value.
 	 * @param Name - Name of bus
 	 * @param Parameter - Default value for created bus
-	 * @param Activate - Whether or not to activate bus on creation. If true, deactivation will only occur
-	 * if returned bus is manually deactivated and not referenced or destroyed (i.e. will not deactivate
-	 * when all references become inactive).
-	 * @return Capture this in a Blueprint variable to prevent it from being automatically garbage collected. 
+	 * @param Activate - (DEPRECATED in 5.4: Use UAudioModulationDestination) Whether or not to activate bus
+	 * on creation. If true, deactivation will only occur if returned bus is manually deactivated and not referenced
+	 * or destroyed (i.e. will not deactivate when all references become inactive).
+	 * @return ControlBus created.  This should be stored (eg. by a Blueprint as a variable) to prevent it from being garbage collected. 
 	 */
 	UFUNCTION(BlueprintCallable, Category = "Audio|Modulation", DisplayName = "Create Control Bus", meta = (
 		AdvancedDisplay = "3",
@@ -134,11 +134,11 @@ public:
 		WorldContext = "WorldContextObject",
 		Keywords = "make modulator")
 	)
-		static UPARAM(DisplayName = "Parameter") USoundModulationParameter* CreateModulationParameter(
-			UObject* WorldContextObject,
-			FName Name,
-			TSubclassOf<USoundModulationParameter> ParamClass,
-			float DefaultValue);
+	static UPARAM(DisplayName = "Parameter") USoundModulationParameter* CreateModulationParameter(
+		UObject* WorldContextObject,
+		FName Name,
+		TSubclassOf<USoundModulationParameter> ParamClass,
+		float DefaultValue);
 
 	/** Creates a modulation generator based on an Envelope Follower with the given parameters.
 	 * @param Name - Name of generator.
@@ -149,10 +149,10 @@ public:
 		WorldContext = "WorldContextObject",
 		Keywords = "make modulator")
 	)
-		static UPARAM(DisplayName = "Generator") USoundModulationGeneratorEnvelopeFollower* CreateEnvelopeFollowerGenerator(
-			UObject* WorldContextObject,
-			FName Name,
-			FEnvelopeFollowerGeneratorParams Params);
+	static UPARAM(DisplayName = "Generator") USoundModulationGeneratorEnvelopeFollower* CreateEnvelopeFollowerGenerator(
+		UObject* WorldContextObject,
+		FName Name,
+		FEnvelopeFollowerGeneratorParams Params);
 
 	/** Creates a modulation generator based on an LFO with the given parameters.
 	 * @param Name - Name of generator.
@@ -182,16 +182,26 @@ public:
 		FName Name,
 		const FSoundModulationADEnvelopeParams& Params);
 
-	/** Creates a modulation watcher, which activates the given modulator and provides a function to retrieve
-	 * the last value computed of the given modulator on the modulation processing thread.
+	/** Creates a modulation destination, which activates the given modulator (if not already active) 
+	 * and provides a function to retrieve the last value computed of the given modulator on the modulation
+	 * processing thread.
 	 */
 	UFUNCTION(BlueprintCallable, Category = "Audio|Modulation", meta = (WorldContext = "WorldContextObject", Keywords = "modulator watch watcher"))
-	static UPARAM(DisplayName = "Watcher") USoundModulationWatcher* CreateModulationWatcher(
+	static UPARAM(DisplayName = "Destination") UAudioModulationDestination* CreateModulationDestination(
 		UObject* WorldContextObject,
 		FName Name,
 		USoundModulatorBase* Modulator);
 
-	/** SOFT DEPRECATED: Use CreateModulationWatcher and store resulting watcher instead!
+	UFUNCTION(BlueprintCallable, Category = "Audio|Modulation", meta = (
+		DisplayName = "Create Modulation Watcher (Deprecated - 5.4)",
+		WorldContext = "WorldContextObject",
+		Keywords = "modulator watch watcher"))
+	static UPARAM(DisplayName = "Watcher") UAudioModulationDestination* CreateModulationWatcher(
+		UObject* WorldContextObject,
+		FName Name,
+		USoundModulatorBase* Modulator);
+
+	/** SOFT DEPRECATED: Use CreateModulationDestination and store resulting destination instead!
 	 * Deactivates a bus. Does nothing if the provided bus is already inactive.
 	 * @param Bus - Scope of modulator
 	 */
@@ -210,12 +220,12 @@ public:
 	)
 	static void DeactivateBusMix(const UObject* WorldContextObject, USoundControlBusMix* Mix);
 
-	/** SOFT DEPRECATED: Use CreateModulationWatcher and store resulting watcher instead!
+	/** SOFT DEPRECATED: Use CreateModulationDestination and store resulting destination instead!
 	 * Deactivates a modulation generator. Does nothing if an instance of the provided generator is already inactive.
 	 * @param Generator - Generator to activate
 	 * @param Scope - Scope of modulator
 	 */
-	UFUNCTION(BlueprintCallable, Category = "Audio|Modulation", DisplayName = "Deactivate Modulation Generator", meta = (
+	UFUNCTION(BlueprintCallable, Category = "Audio|Modulation", DisplayName = "Deactivate Modulation Generator (Deprecated - 5.4)", meta = (
 		WorldContext = "WorldContextObject",
 		Keywords = "bus modulation modulator generator")
 	)
diff --git a/Engine/Plugins/Runtime/AudioModulation/Source/AudioModulation/Public/SoundModulationWatcher.h b/Engine/Plugins/Runtime/AudioModulation/Source/AudioModulation/Public/SoundModulationWatcher.h
deleted file mode 100644
