Properly check for valid CBM in SetGlobalBusMixValue

diff --git a/Engine/Plugins/Runtime/AudioModulation/Source/AudioModulation/Private/AudioModulationSystem.cpp b/Engine/Plugins/Runtime/AudioModulation/Source/AudioModulation/Private/AudioModulationSystem.cpp
--- a/Engine/Plugins/Runtime/AudioModulation/Source/AudioModulation/Private/AudioModulationSystem.cpp
+++ b/Engine/Plugins/Runtime/AudioModulation/Source/AudioModulation/Private/AudioModulationSystem.cpp
@@ -229,8 +229,10 @@ namespace AudioModulation
 
 	void FAudioModulationSystem::SetGlobalBusMixValue(USoundControlBus& InBus, float InValue, float InFadeTime)
 	{
-		if (TObjectPtr<USoundControlBusMix> GlobalMix = ActiveGlobalBusValueMixes.FindRef(InBus.GetUniqueID()))
+		if (ActiveGlobalBusValueMixes.Contains(InBus.GetUniqueID()))
 		{
+			TObjectPtr<USoundControlBusMix> GlobalMix = ActiveGlobalBusValueMixes.FindRef(InBus.GetUniqueID());
+
 			if (ensure(!GlobalMix->MixStages.IsEmpty()))
 			{
 				GlobalMix->MixStages[0].Value.TargetValue = InValue;
