Is Control Bus Mix Active BP function

diff --git a/Engine/Plugins/Runtime/AudioModulation/Source/AudioModulation/Private/AudioModulation.cpp b/Engine/Plugins/Runtime/AudioModulation/Source/AudioModulation/Private/AudioModulation.cpp
--- a/Engine/Plugins/Runtime/AudioModulation/Source/AudioModulation/Private/AudioModulation.cpp
+++ b/Engine/Plugins/Runtime/AudioModulation/Source/AudioModulation/Private/AudioModulation.cpp
@@ -97,6 +97,11 @@ namespace AudioModulation
 		PRAGMA_ENABLE_DEPRECATION_WARNINGS
 	}
 
+	bool FAudioModulationManager::IsBusMixActive(const USoundControlBusMix& InBusMix)
+	{
+		return ModSystem->IsControlBusMixActive(InBusMix);
+	}
+
 #if !UE_BUILD_SHIPPING
 	void FAudioModulationManager::SetDebugBusFilter(const FString* InNameFilter)
 	{
diff --git a/Engine/Plugins/Runtime/AudioModulation/Source/AudioModulation/Private/AudioModulationStatics.cpp b/Engine/Plugins/Runtime/AudioModulation/Source/AudioModulation/Private/AudioModulationStatics.cpp
--- a/Engine/Plugins/Runtime/AudioModulation/Source/AudioModulation/Private/AudioModulationStatics.cpp
+++ b/Engine/Plugins/Runtime/AudioModulation/Source/AudioModulation/Private/AudioModulationStatics.cpp
@@ -391,6 +391,19 @@ void UAudioModulationStatics::DeactivateGenerator(const UObject* WorldContextObj
 	}
 }
 
+bool UAudioModulationStatics::IsControlBusMixActive(const UObject* WorldContextObject, USoundControlBusMix* Mix)
+{
+	UWorld* World = GetAudioWorld(WorldContextObject);
+	if (AudioModulation::FAudioModulationManager* ModSystem = GetModulation(World))
+	{
+		if (Mix)
+		{
+			return ModSystem->IsBusMixActive(*Mix);
+		}
+	}
+	return false;
+}
+
 void UAudioModulationStatics::SaveMixToProfile(const UObject* WorldContextObject, USoundControlBusMix* BusMix, int32 ProfileIndex)
 {
 	UWorld* World = GetAudioWorld(WorldContextObject);
diff --git a/Engine/Plugins/Runtime/AudioModulation/Source/AudioModulation/Private/AudioModulationSystem.cpp b/Engine/Plugins/Runtime/AudioModulation/Source/AudioModulation/Private/AudioModulationSystem.cpp
--- a/Engine/Plugins/Runtime/AudioModulation/Source/AudioModulation/Private/AudioModulationSystem.cpp
+++ b/Engine/Plugins/Runtime/AudioModulation/Source/AudioModulation/Private/AudioModulationSystem.cpp
@@ -91,6 +91,8 @@ namespace AudioModulation
 
 	void FAudioModulationSystem::ActivateBusMix(FModulatorBusMixSettings&& InSettings)
 	{
+		ActiveBusMixIds.Add(InSettings.GetId());
+
 		RunCommandOnProcessingThread([this, Settings = MoveTemp(InSettings)]() mutable
 		{
 			FBusMixHandle BusMixHandle = FBusMixHandle::Get(Settings.GetId(), RefProxies.BusMixes);
@@ -159,6 +161,9 @@ namespace AudioModulation
 
 	void FAudioModulationSystem::DeactivateBusMix(const USoundControlBusMix& InBusMix)
 	{
+		FBusMixHandle MixHandle = FBusMixHandle::Get(InBusMix.GetUniqueID(), RefProxies.BusMixes);
+		ActiveBusMixIds.Remove(MixHandle.GetId());
+		
 		RunCommandOnProcessingThread([this, BusMixId = static_cast<FBusMixId>(InBusMix.GetUniqueID())]()
 		{
 			FBusMixHandle MixHandle = FBusMixHandle::Get(BusMixId, RefProxies.BusMixes);
@@ -174,6 +179,8 @@ namespace AudioModulation
 	{
 		ClearAllGlobalBusMixValues();
 
+		ActiveBusMixIds.Empty();
+
 		RunCommandOnProcessingThread([this]()
 		{
 			for (TPair<FBusMixId, FModulatorBusMixProxy>& Pair : RefProxies.BusMixes)
@@ -518,6 +525,12 @@ namespace AudioModulation
 #endif // !UE_BUILD_SHIPPING
 	}
 
+	bool FAudioModulationSystem::IsControlBusMixActive(const USoundControlBusMix& InBusMix)
+	{
+		FBusMixId BusMixId = static_cast<FBusMixId>(InBusMix.GetUniqueID());		
+		return ActiveBusMixIds.Contains(BusMixId);
+	}
+
 	void FAudioModulationSystem::SaveMixToProfile(const USoundControlBusMix& InBusMix, const int32 InProfileIndex)
 	{
 		check(IsInGameThread());
diff --git a/Engine/Plugins/Runtime/AudioModulation/Source/AudioModulation/Private/AudioModulationSystem.h b/Engine/Plugins/Runtime/AudioModulation/Source/AudioModulation/Private/AudioModulationSystem.h
--- a/Engine/Plugins/Runtime/AudioModulation/Source/AudioModulation/Private/AudioModulationSystem.h
+++ b/Engine/Plugins/Runtime/AudioModulation/Source/AudioModulation/Private/AudioModulationSystem.h
@@ -114,6 +114,8 @@ namespace AudioModulation
 
 		void UnregisterModulator(const Audio::FModulatorHandle& InHandle);
 
+		bool IsControlBusMixActive(const USoundControlBusMix& InBusMix);
+
 		/* Saves mix to .ini profile for fast iterative development that does not require re-cooking a mix */
 		void SaveMixToProfile(const USoundControlBusMix& InBusMix, const int32 InProfileIndex);
 
@@ -223,6 +225,8 @@ namespace AudioModulation
 		mutable FCriticalSection ThreadSafeModValueCritSection;
 		TMap<Audio::FModulatorId, float> ThreadSafeModValueMap;
 
+		TSet<FBusMixId> ActiveBusMixIds;
+
 		TSet<FBusHandle> ManuallyActivatedBuses;
 		TSet<FBusMixHandle> ManuallyActivatedBusMixes;
 		TSet<FGeneratorHandle> ManuallyActivatedGenerators;
@@ -305,6 +309,8 @@ namespace AudioModulation
 
 		Audio::FDeviceId GetAudioDeviceId() const { return 0; }
 
+		bool IsControlBusMixActive(const USoundControlBusMix& InBusMix) { return false; }
+
 		void SaveMixToProfile(const USoundControlBusMix& InBusMix, const int32 InProfileIndex) { }
 		TArray<FSoundControlBusMixStage> LoadMixFromProfile(const int32 InProfileIndex, USoundControlBusMix& OutBusMix) { return { }; }
 
@@ -325,6 +331,11 @@ namespace AudioModulation
 		bool GetModulatorValueThreadSafe(uint32 ModulatorID, float& OutValue) const { return false; }
 		void UnregisterModulator(const Audio::FModulatorHandle& InHandle) { }
 
+		inline bool FAudioModulationSystem::IsControlBusMixActive(const USoundControlBusMix* InBusMix)
+		{
+			return false;
+		}
+
 		void UpdateMix(const USoundControlBusMix& InMix, float InFadeTime = -1.0f) { }
 		void UpdateMix(const TArray<FSoundControlBusMixStage>& InStages, USoundControlBusMix& InOutMix, bool bUpdateObject = false, float InFadeTime = -1.0f) { }
 		void UpdateMixByFilter(const FString& InAddressFilter, const TSubclassOf<USoundModulationParameter>& InParamClassFilter, USoundModulationParameter* InParamFilter, float InValue, float InFadeTime, USoundControlBusMix& InOutMix, bool bUpdateObject = false) { }
diff --git a/Engine/Plugins/Runtime/AudioModulation/Source/AudioModulation/Public/AudioModulation.h b/Engine/Plugins/Runtime/AudioModulation/Source/AudioModulation/Public/AudioModulation.h
--- a/Engine/Plugins/Runtime/AudioModulation/Source/AudioModulation/Public/AudioModulation.h
+++ b/Engine/Plugins/Runtime/AudioModulation/Source/AudioModulation/Public/AudioModulation.h
@@ -76,6 +76,8 @@ namespace AudioModulation
 		UE_DEPRECATED(5.4, "Deactivation of modulators in this manner is now deprecated. Use USoundModulationWatchers to safety activate and track a given modulator")
 		void DeactivateGenerator(const USoundModulationGenerator& InGenerator);
 
+		bool IsBusMixActive(const USoundControlBusMix& InBusMix);
+
 		void SaveMixToProfile(const USoundControlBusMix& InBusMix, const int32 InProfileIndex);
 		TArray<FSoundControlBusMixStage> LoadMixFromProfile(const int32 InProfileIndex, USoundControlBusMix& OutBusMix);
 
diff --git a/Engine/Plugins/Runtime/AudioModulation/Source/AudioModulation/Public/AudioModulationStatics.h b/Engine/Plugins/Runtime/AudioModulation/Source/AudioModulation/Public/AudioModulationStatics.h
--- a/Engine/Plugins/Runtime/AudioModulation/Source/AudioModulation/Public/AudioModulationStatics.h
+++ b/Engine/Plugins/Runtime/AudioModulation/Source/AudioModulation/Public/AudioModulationStatics.h
@@ -236,6 +236,16 @@ public:
 	)
 	static void DeactivateGenerator(const UObject* WorldContextObject, USoundModulationGenerator* Generator);
 
+	/** Returns whether or not a Control Bus Mix is currently active.
+	 * @param Mix - the Control Bus Mix to check.
+	 * @return Whether or not the Bus Mix is currently active.
+	 */
+	UFUNCTION(BlueprintPure, Category = "Audio|Modulation", DisplayName = "Is Control Bus Mix Active", meta = (
+		WorldContext = "WorldContextObject",
+		Keywords = "bus modulation modulator generator")
+		)
+	static UPARAM(DisplayName = "Is Active") bool IsControlBusMixActive(const UObject * WorldContextObject, USoundControlBusMix * Mix);
+
 	/** Saves control bus mix to a profile, serialized to an ini file.  If mix is loaded, uses current proxy's state.
 	 * If not, uses default UObject representation.
 	 * @param BusMix - Mix object to serialize to profile .ini.
