Clean-up Global Mix implementation
- Track global bus mixes in modulation system using ref collector instead of adding to root
- Track mixes directly within the audio worldization subsystem instead of using global mixes

diff --git a/Engine/Plugins/Runtime/AudioModulation/Source/AudioModulation/Private/AudioModulation.cpp b/Engine/Plugins/Runtime/AudioModulation/Source/AudioModulation/Private/AudioModulation.cpp
--- a/Engine/Plugins/Runtime/AudioModulation/Source/AudioModulation/Private/AudioModulation.cpp
+++ b/Engine/Plugins/Runtime/AudioModulation/Source/AudioModulation/Private/AudioModulation.cpp
@@ -69,6 +69,11 @@ namespace AudioModulation
 		PRAGMA_ENABLE_DEPRECATION_WARNINGS
 	}
 
+	USoundControlBusMix* FAudioModulationManager::CreateBusMixFromValue(FName Name, const TArray<USoundControlBus*>& Buses, float Value, float AttackTime, float ReleaseTime)
+	{
+		return ModSystem->CreateBusMixFromValue(Name, Buses, Value, AttackTime, ReleaseTime);
+	}
+
 	void FAudioModulationManager::DeactivateBus(const USoundControlBus& InBus)
 	{
 		UE_LOG(LogAudioModulation, Warning, TEXT("FAudioModulationManager::DeactivateBus is deprecated. Use UAudioModulationDestination API (see ClearModulator) instead to enforce bus lifetime."));
diff --git a/Engine/Plugins/Runtime/AudioModulation/Source/AudioModulation/Private/AudioModulationStatics.cpp b/Engine/Plugins/Runtime/AudioModulation/Source/AudioModulation/Private/AudioModulationStatics.cpp
--- a/Engine/Plugins/Runtime/AudioModulation/Source/AudioModulation/Private/AudioModulationStatics.cpp
+++ b/Engine/Plugins/Runtime/AudioModulation/Source/AudioModulation/Private/AudioModulationStatics.cpp
@@ -337,6 +337,12 @@ USoundModulationGeneratorADEnvelope* UAudioModulationStatics::CreateADEnvelopeGe
 
 UAudioModulationDestination* UAudioModulationStatics::CreateModulationDestination(UObject* WorldContextObject, FName Name, USoundModulatorBase* Modulator)
 {
+	UWorld* World = GetAudioWorld(WorldContextObject);
+	if (!World)
+	{
+		return nullptr;
+	}
+
 	if (UAudioModulationDestination* NewDestination = NewObject<UAudioModulationDestination>(WorldContextObject, Name))
 	{
 		NewDestination->SetModulator(Modulator);
@@ -448,6 +454,25 @@ void UAudioModulationStatics::UpdateMix(const UObject* WorldContextObject, USoun
 	}
 }
 
+USoundControlBusMix* UAudioModulationStatics::CreateBusMixFromValue(const UObject* WorldContextObject, FName Name, const TArray<USoundControlBus*>& Buses, float Value, float AttackTime, float ReleaseTime, bool bActivate)
+{
+	UWorld* World = GetAudioWorld(WorldContextObject);
+	if (AudioModulation::FAudioModulationManager* ModSystem = GetModulation(World))
+	{
+		if (USoundControlBusMix* NewMix = ModSystem->CreateBusMixFromValue(Name, Buses, Value, AttackTime, ReleaseTime))
+		{
+			if (bActivate)
+			{
+				ModSystem->ActivateBusMix(*NewMix);
+			}
+
+			return NewMix;
+		}
+	}
+
+	return nullptr;
+}
+
 void UAudioModulationStatics::SetGlobalBusMixValue(const UObject* WorldContextObject, USoundControlBus* Bus, float Value, float FadeTime)
 {
 	if (Bus)
diff --git a/Engine/Plugins/Runtime/AudioModulation/Source/AudioModulation/Private/AudioModulationSystem.cpp b/Engine/Plugins/Runtime/AudioModulation/Source/AudioModulation/Private/AudioModulationSystem.cpp
--- a/Engine/Plugins/Runtime/AudioModulation/Source/AudioModulation/Private/AudioModulationSystem.cpp
+++ b/Engine/Plugins/Runtime/AudioModulation/Source/AudioModulation/Private/AudioModulationSystem.cpp
@@ -131,6 +131,18 @@ namespace AudioModulation
 		});
 	}
 
+	void FAudioModulationSystem::AddReferencedObjects(FReferenceCollector& Collector)
+	{
+		TArray<TObjectPtr<USoundControlBusMix>> GlobalBusMixes;
+		ActiveGlobalBusValueMixes.GenerateValueArray(GlobalBusMixes);
+		Collector.AddReferencedObjects(GlobalBusMixes);
+	}
+
+	FString FAudioModulationSystem::GetReferencerName() const
+	{
+		return TEXT("FAudioModulationSystem");
+	}
+
 	bool FAudioModulationSystem::CalculateModulationValue(FModulationPatchProxy& OutProxy, float& OutValue) const
 	{
 		check(IsInProcessingThread());
@@ -250,29 +262,12 @@ namespace AudioModulation
 		}
 		else
 		{
-			const FString MixName = InBus.GetName() + TEXT("_GlobalMix");
-			if (TObjectPtr<USoundControlBusMix> NewGlobalMix = NewObject<USoundControlBusMix>(GetTransientPackage(), FName(*MixName)))
+			const FName MixName(*FString::Printf(TEXT("%s_GlobalMix"), *InBus.GetName()));
+			if (TObjectPtr<USoundControlBusMix> NewGlobalMix = CreateBusMixFromValue(MixName, { &InBus }, InValue, InFadeTime))
 			{
-				{
-					FSoundModulationMixValue MixValue;
-					MixValue.TargetValue = InValue;
-
-					if (InFadeTime >= 0.0f)
-					{
-						MixValue.AttackTime = InFadeTime;
-					}
-
-					FSoundControlBusMixStage MixStage;
-					MixStage.Bus = &InBus;
-					MixStage.Value = MixValue;
-
-					NewGlobalMix->MixStages.Emplace(MoveTemp(MixStage));
-				}
-
 				ActiveGlobalBusValueMixes.Add(InBus.GetUniqueID(), NewGlobalMix);
 				UE_LOG(LogAudioModulation, VeryVerbose, TEXT("GlobalBusMix for ControlBus '%s' activated, target set to '%0.4f'."), *InBus.GetName(), InValue);
 				ActivateBusMix(*NewGlobalMix);
-				NewGlobalMix->AddToRoot();
 			}
 			else
 			{
@@ -292,7 +287,6 @@ namespace AudioModulation
 				GlobalMix->MixStages[0].Value.ReleaseTime = InFadeTime;
 				DeactivateBusMix(*GlobalMix);
 				ActiveGlobalBusValueMixes.Remove(BusID);
-				GlobalMix->RemoveFromRoot();
 				UE_LOG(LogAudioModulation, VeryVerbose, TEXT("GlobalBusMix for ControlBus '%s' cleared."), *InBus.GetName());
 			}
 		}
@@ -317,6 +311,40 @@ namespace AudioModulation
 		ActiveGlobalBusValueMixes.Reset();
 	}
 
+	USoundControlBusMix* FAudioModulationSystem::CreateBusMixFromValue(FName Name, const TArray<USoundControlBus*>& Buses, float Value, float AttackTime, float ReleaseTime)
+	{
+		if (TObjectPtr<USoundControlBusMix> NewGlobalMix = NewObject<USoundControlBusMix>(GetTransientPackage(), Name))
+		{
+			for (USoundControlBus* Bus : Buses)
+			{
+				if (Bus)
+				{
+					FSoundModulationMixValue MixValue;
+					MixValue.TargetValue = Value;
+
+					if (AttackTime >= 0.0f)
+					{
+						MixValue.AttackTime = AttackTime;
+					}
+
+					if (ReleaseTime >= 0.0f)
+					{
+						MixValue.ReleaseTime = ReleaseTime;
+					}
+
+					FSoundControlBusMixStage MixStage;
+					MixStage.Bus = Bus;
+					MixStage.Value = MixValue;
+
+					NewGlobalMix->MixStages.Emplace(MoveTemp(MixStage));
+				}
+			}
+			return NewGlobalMix;
+		}
+
+		return nullptr;
+	}
+
 	bool FAudioModulationSystem::GetModulatorValue(const Audio::FModulatorHandle& InModulatorHandle, float& OutValue) const
 	{
 		const EModulatorType ModulatorType = static_cast<EModulatorType>(InModulatorHandle.GetTypeId());
diff --git a/Engine/Plugins/Runtime/AudioModulation/Source/AudioModulation/Private/AudioModulationSystem.h b/Engine/Plugins/Runtime/AudioModulation/Source/AudioModulation/Private/AudioModulationSystem.h
--- a/Engine/Plugins/Runtime/AudioModulation/Source/AudioModulation/Private/AudioModulationSystem.h
+++ b/Engine/Plugins/Runtime/AudioModulation/Source/AudioModulation/Private/AudioModulationSystem.h
@@ -20,6 +20,7 @@
 #include "SoundModulationGeneratorProxy.h"
 #include "Templates/Atomic.h"
 #include "Templates/Function.h"
+#include "UObject/GCObject.h"
 
 #if WITH_AUDIOMODULATION
 #if !UE_BUILD_SHIPPING
@@ -62,7 +63,7 @@ namespace AudioModulation
 #if WITH_AUDIOMODULATION
 namespace AudioModulation
 {
-	class FAudioModulationSystem
+	class FAudioModulationSystem : public FGCObject
 	{
 	public:
 		void Initialize(const FAudioPluginInitializationParams& InitializationParams);
@@ -76,6 +77,11 @@ namespace AudioModulation
 		UE_DEPRECATED(5.4, "Activation of modulators in this manner is now deprecated. Use USoundModulationWatchers to safety activate and track a given modulator")
 		void ActivateGenerator(const USoundModulationGenerator& InGenerator);
 
+		// FGCObject interface
+		virtual void AddReferencedObjects(FReferenceCollector& Collector) override;
+		virtual FString GetReferencerName() const override;
+		// End of FGCObject interface
+
 		UE_DEPRECATED(5.4, "Deactivation of modulators in this manner is now deprecated. Use USoundModulationWatchers to safety activate and track a given modulator")
 		void DeactivateBus(const USoundControlBus& InBus);
 
@@ -144,6 +150,9 @@ namespace AudioModulation
 		/* Clears all global bus mix values over the prescribed FadeTime. If FadeTime is non-positive, returns to the bus's respective parameter default immediately. */
 		void ClearAllGlobalBusMixValues(float FadeTime = -1.0f);
 
+		/* Create a mix from an array of buses where a supplied default value and associated timing parameters are applied for each bus's stage */
+		USoundControlBusMix* CreateBusMixFromValue(FName Name, const TArray<USoundControlBus*>& Buses, float Value, float AttackTime = -1.0f, float ReleaseTime = -1.0f);
+
 		/*
 		 * Commits any changes from a modulator type applied to a UObject definition
 		 * to modulator instance if active (i.e. Control Bus, Control Bus Modulator)
@@ -318,6 +327,8 @@ namespace AudioModulation
 		void ClearGlobalBusMixValue(const USoundControlBus& InBus, float FadeTime) { }
 		void ClearAllGlobalBusMixValues(float FadeTime) { }
 
+		USoundControlBusMix* CreateBusMixFromValue(FName Name, const TArray<USoundControlBus*>& Buses, float Value, float AttackTime = -1.0f, float ReleaseTime = -1.0f) { return nullptr; }
+
 		void ProcessModulators(const double InElapsed) { }
 		void SoloBusMix(const USoundControlBusMix& InBusMix) { }
 
diff --git a/Engine/Plugins/Runtime/AudioModulation/Source/AudioModulation/Public/AudioModulation.h b/Engine/Plugins/Runtime/AudioModulation/Source/AudioModulation/Public/AudioModulation.h
--- a/Engine/Plugins/Runtime/AudioModulation/Source/AudioModulation/Public/AudioModulation.h
+++ b/Engine/Plugins/Runtime/AudioModulation/Source/AudioModulation/Public/AudioModulation.h
@@ -67,6 +67,8 @@ namespace AudioModulation
 		UE_DEPRECATED(5.4, "Deactivation of modulators in this manner is now deprecated. Use USoundModulationWatchers to safety activate and track a given modulator")
 		void ActivateGenerator(const USoundModulationGenerator& InGenerator);
 
+		USoundControlBusMix* CreateBusMixFromValue(FName Name, const TArray<USoundControlBus*>& Buses, float Value, float AttackTime = -1.0f, float ReleaseTime = -1.0f);
+
 		UE_DEPRECATED(5.4, "Deactivation of modulators in this manner is now deprecated. Use USoundModulationWatchers to safety activate and track a given modulator")
 		void DeactivateBus(const USoundControlBus& InBus);
 
diff --git a/Engine/Plugins/Runtime/AudioModulation/Source/AudioModulation/Public/AudioModulationStatics.h b/Engine/Plugins/Runtime/AudioModulation/Source/AudioModulation/Public/AudioModulationStatics.h
--- a/Engine/Plugins/Runtime/AudioModulation/Source/AudioModulation/Public/AudioModulationStatics.h
+++ b/Engine/Plugins/Runtime/AudioModulation/Source/AudioModulation/Public/AudioModulationStatics.h
@@ -88,6 +88,26 @@ public:
 	)
 	static UPARAM(DisplayName = "Bus") USoundControlBus* CreateBus(UObject* WorldContextObject, FName Name, USoundModulationParameter* Parameter, bool Activate = true);
 
+	/* Create a mix with stages created for each provided bus that are initialized to the supplied value and timing parameters.
+	 * @param Buses - Buses to assign stages within new mix to
+	 * @param Value - Initial value for all stages created within the new mix.
+	 * @param AttackTime - Fade time to user when mix activates.
+	 * @param ReleaseTime - Fade time to user when mix deactivates.
+	 */
+	UFUNCTION(BlueprintCallable, Category = "Audio|Modulation", meta = (
+		AdvancedDisplay = "3",
+		WorldContext = "WorldContextObject",
+		Keywords = "make modulation modulator stage")
+	)
+	static UPARAM(DisplayName = "Mix") USoundControlBusMix* CreateBusMixFromValue(
+		const UObject* WorldContextObject,
+		FName Name,
+		const TArray<USoundControlBus*>& Buses,
+		float Value = 1.0f,
+		float AttackTime = 0.1f,
+		float ReleaseTime = 0.1f,
+		bool bActivate = true);
+
 	/** Creates a stage used to mix a control bus.
 	 * @param Bus - Bus stage is in charge of applying mix value to.
 	 * @param Value - Value for added bus stage to target when mix is active.
