Fix deprecation warning to be more accurate now that watchers have been redirected
#rnx

diff --git a/Engine/Plugins/Runtime/AudioModulation/Source/AudioModulation/Private/AudioModulationSystem.h b/Engine/Plugins/Runtime/AudioModulation/Source/AudioModulation/Private/AudioModulationSystem.h
--- a/Engine/Plugins/Runtime/AudioModulation/Source/AudioModulation/Private/AudioModulationSystem.h
+++ b/Engine/Plugins/Runtime/AudioModulation/Source/AudioModulation/Private/AudioModulationSystem.h
@@ -68,13 +68,13 @@ namespace AudioModulation
 	public:
 		void Initialize(const FAudioPluginInitializationParams& InitializationParams);
 
-		UE_DEPRECATED(5.4, "Activation of modulators in this manner is now deprecated. Use USoundModulationWatchers to safety activate and track a given modulator")
+		UE_DEPRECATED(5.4, "Activation of modulators in this manner is now deprecated. Use USoundModulationDestinations to safely activate and track a given modulator")
 		void ActivateBus(const USoundControlBus& InBus);
 
 		void ActivateBusMix(FModulatorBusMixSettings&& InSettings);
 		void ActivateBusMix(const USoundControlBusMix& InBusMix);
 
-		UE_DEPRECATED(5.4, "Activation of modulators in this manner is now deprecated. Use USoundModulationWatchers to safety activate and track a given modulator")
+		UE_DEPRECATED(5.4, "Activation of modulators in this manner is now deprecated. Use USoundModulationDestinations to safely activate and track a given modulator")
 		void ActivateGenerator(const USoundModulationGenerator& InGenerator);
 
 		// FGCObject interface
@@ -82,7 +82,7 @@ namespace AudioModulation
 		virtual FString GetReferencerName() const override;
 		// End of FGCObject interface
 
-		UE_DEPRECATED(5.4, "Deactivation of modulators in this manner is now deprecated. Use USoundModulationWatchers to safety activate and track a given modulator")
+		UE_DEPRECATED(5.4, "Deactivation of modulators in this manner is now deprecated. Use USoundModulationDestinations to safely activate and track a given modulator")
 		void DeactivateBus(const USoundControlBus& InBus);
 
 		/** Deactivates given bus mix */
@@ -91,7 +91,7 @@ namespace AudioModulation
 		/** Deactivates all bus mixes */
 		void DeactivateAllBusMixes();
 
-		UE_DEPRECATED(5.4, "Deactivation of modulators in this manner is now deprecated. Use USoundModulationWatchers to safety activate and track a given modulator")
+		UE_DEPRECATED(5.4, "Deactivation of modulators in this manner is now deprecated. Use USoundModulationDestinations to safely activate and track a given modulator")
 		void DeactivateGenerator(const USoundModulationGenerator& InGenerator);
 
 		void ProcessModulators(const double InElapsed);
