Remove Deprecation-related issues that came from BP Modulation Destination

diff --git a/Engine/Plugins/Runtime/AudioModulation/Source/AudioModulation/Private/AudioModulationStatics.cpp b/Engine/Plugins/Runtime/AudioModulation/Source/AudioModulation/Private/AudioModulationStatics.cpp
--- a/Engine/Plugins/Runtime/AudioModulation/Source/AudioModulation/Private/AudioModulationStatics.cpp
+++ b/Engine/Plugins/Runtime/AudioModulation/Source/AudioModulation/Private/AudioModulationStatics.cpp
@@ -352,11 +352,6 @@ UAudioModulationDestination* UAudioModulationStatics::CreateModulationDestinatio
 	return nullptr;
 }
 
-UAudioModulationDestination* UAudioModulationStatics::CreateModulationWatcher(UObject* WorldContextObject, FName Name, USoundModulatorBase* Modulator)
-{
-	return CreateModulationDestination(WorldContextObject, Name, Modulator);
-}
-
 void UAudioModulationStatics::DeactivateBus(const UObject* WorldContextObject, USoundControlBus* Bus)
 {
 	if (Bus)
diff --git a/Engine/Plugins/Runtime/AudioModulation/Source/AudioModulation/Public/AudioModulationStatics.h b/Engine/Plugins/Runtime/AudioModulation/Source/AudioModulation/Public/AudioModulationStatics.h
--- a/Engine/Plugins/Runtime/AudioModulation/Source/AudioModulation/Public/AudioModulationStatics.h
+++ b/Engine/Plugins/Runtime/AudioModulation/Source/AudioModulation/Public/AudioModulationStatics.h
@@ -86,7 +86,7 @@ public:
 		WorldContext = "WorldContextObject",
 		Keywords = "make modulation LPF modulator")
 	)
-	static UPARAM(DisplayName = "Bus") USoundControlBus* CreateBus(UObject* WorldContextObject, FName Name, USoundModulationParameter* Parameter, bool Activate = true);
+	static UPARAM(DisplayName = "Bus") USoundControlBus* CreateBus(UObject* WorldContextObject, FName Name, USoundModulationParameter* Parameter, bool Activate = false);
 
 	/* Create a mix with stages created for each provided bus that are initialized to the supplied value and timing parameters.
 	 * @param Buses - Buses to assign stages within new mix to
@@ -212,20 +212,6 @@ public:
 		FName Name,
 		USoundModulatorBase* Modulator);
 
-	/** SOFT DEPRECATED: Use CreateModulationDestination and store resulting destination instead!
-	 * Creates a modulation destination, which activates the given modulator (if not already active)
-	 * and provides a function to retrieve the last value computed of the given modulator on the modulation
-	 * processing thread.
-	 */
-	UFUNCTION(BlueprintCallable, Category = "Audio|Modulation", meta = (
-		DisplayName = "Create Modulation Watcher (Deprecated - 5.4)",
-		WorldContext = "WorldContextObject",
-		Keywords = "modulator watch watcher"))
-	static UPARAM(DisplayName = "Watcher") UAudioModulationDestination* CreateModulationWatcher(
-		UObject* WorldContextObject,
-		FName Name,
-		USoundModulatorBase* Modulator);
-
 	/** SOFT DEPRECATED: Use CreateModulationDestination and store resulting destination instead!
 	 * Deactivates a bus. Does nothing if the provided bus is already inactive.
 	 * @param Bus - Scope of modulator
