Fix crash when access audio bus subsystem from modulator.

diff --git a/Engine/Plugins/Runtime/AudioModulation/Source/AudioModulation/Private/Generators/SoundModulationEnvelopeFollower.cpp b/Engine/Plugins/Runtime/AudioModulation/Source/AudioModulation/Private/Generators/SoundModulationEnvelopeFollower.cpp
--- a/Engine/Plugins/Runtime/AudioModulation/Source/AudioModulation/Private/Generators/SoundModulationEnvelopeFollower.cpp
+++ b/Engine/Plugins/Runtime/AudioModulation/Source/AudioModulation/Private/Generators/SoundModulationEnvelopeFollower.cpp
@@ -212,10 +212,11 @@ namespace AudioModulation
 					FAudioDevice* AudioDevice = DeviceManager->GetAudioDeviceRaw(AudioDeviceId);
 					if (Audio::FMixerDevice* MixerDevice = static_cast<Audio::FMixerDevice*>(AudioDevice))
 					{
-						UAudioBusSubsystem* AudioBusSubsystem = MixerDevice->GetSubsystem<UAudioBusSubsystem>();
-						check(AudioBusSubsystem);
-						AudioBusPatch = AudioBusSubsystem->AddPatchOutputForAudioBus(Audio::FAudioBusKey(BusId), MixerDevice->GetNumOutputFrames(), EnvelopeFollower.GetNumChannels(), Gain);
-						bBusRequiresPatch = false;
+						if (UAudioBusSubsystem* AudioBusSubsystem = MixerDevice->GetSubsystem<UAudioBusSubsystem>())
+						{
+							AudioBusPatch = AudioBusSubsystem->AddPatchOutputForAudioBus(Audio::FAudioBusKey(BusId), MixerDevice->GetNumOutputFrames(), EnvelopeFollower.GetNumChannels(), Gain);
+							bBusRequiresPatch = false;
+						}
 					}
 				}
 			}
