Fix for Modulation Envelope Follower Regression

diff --git a/Engine/Plugins/Runtime/AudioModulation/Source/AudioModulation/Private/Generators/SoundModulationEnvelopeFollower.cpp b/Engine/Plugins/Runtime/AudioModulation/Source/AudioModulation/Private/Generators/SoundModulationEnvelopeFollower.cpp
--- a/Engine/Plugins/Runtime/AudioModulation/Source/AudioModulation/Private/Generators/SoundModulationEnvelopeFollower.cpp
+++ b/Engine/Plugins/Runtime/AudioModulation/Source/AudioModulation/Private/Generators/SoundModulationEnvelopeFollower.cpp
@@ -169,7 +169,7 @@ namespace AudioModulation
 				CreatePatchForBus();
 			}
 
-			if (AudioBusPatch.IsValid() && !AudioBusPatch->IsInputStale())
+			if (AudioBusPatch.IsValid())
 			{
 				const int32 NumSamples = AudioBusPatch->GetNumSamplesAvailable();
 				const int32 NumFrames = NumSamples / EnvelopeFollower.GetNumChannels();
