Fix Modulation LFO Inconsistency

diff --git a/Engine/Plugins/Runtime/AudioModulation/Source/AudioModulation/Private/Generators/SoundModulationLFO.cpp b/Engine/Plugins/Runtime/AudioModulation/Source/AudioModulation/Private/Generators/SoundModulationLFO.cpp
--- a/Engine/Plugins/Runtime/AudioModulation/Source/AudioModulation/Private/Generators/SoundModulationLFO.cpp
+++ b/Engine/Plugins/Runtime/AudioModulation/Source/AudioModulation/Private/Generators/SoundModulationLFO.cpp
@@ -47,6 +47,7 @@ namespace AudioModulation
 				LFO.SetGain(Params.Amplitude);
 				LFO.SetMode(Params.bLooping ? Audio::ELFOMode::Type::Sync : Audio::ELFOMode::OneShot);
 				LFO.SetPulseWidth(Params.Width);
+				LFO.SetPhaseOffset(Params.Phase);
 				LFO.SetType(static_cast<Audio::ELFO::Type>(Params.Shape));
 				LFO.Start();
 
