Audio Modulation Insights stub

diff --git a/Engine/Plugins/Runtime/AudioModulation/AudioModulation.uplugin b/Engine/Plugins/Runtime/AudioModulation/AudioModulation.uplugin
--- a/Engine/Plugins/Runtime/AudioModulation/AudioModulation.uplugin
+++ b/Engine/Plugins/Runtime/AudioModulation/AudioModulation.uplugin
@@ -35,6 +35,10 @@
 		{
 			"Name": "WaveTable",
 			"Enabled": true
+		},
+		{
+			"Name": "AudioInsights",
+			"Enabled": true
 		}
 	]
 }
\ No newline at end of file
diff --git a/Engine/Plugins/Runtime/AudioModulation/Source/AudioModulationEditor/AudioModulationEditor.build.cs b/Engine/Plugins/Runtime/AudioModulation/Source/AudioModulationEditor/AudioModulationEditor.build.cs
--- a/Engine/Plugins/Runtime/AudioModulation/Source/AudioModulationEditor/AudioModulationEditor.build.cs
+++ b/Engine/Plugins/Runtime/AudioModulation/Source/AudioModulationEditor/AudioModulationEditor.build.cs
@@ -28,6 +28,7 @@ public class AudioModulationEditor : ModuleRules
 				"InputCore",				
 				"AudioEditor",
 				"AudioExtensions",
+				"AudioInsights",
 				"AudioModulation",
 				"CurveEditor",
 				"EditorWidgets",
diff --git a/Engine/Plugins/Runtime/AudioModulation/Source/AudioModulationEditor/Private/AudioModulationEditor.cpp b/Engine/Plugins/Runtime/AudioModulation/Source/AudioModulationEditor/Private/AudioModulationEditor.cpp
--- a/Engine/Plugins/Runtime/AudioModulation/Source/AudioModulationEditor/Private/AudioModulationEditor.cpp
+++ b/Engine/Plugins/Runtime/AudioModulation/Source/AudioModulationEditor/Private/AudioModulationEditor.cpp
@@ -13,8 +13,10 @@
 #include "Framework/Commands/UIAction.h"
 #include "Framework/MultiBox/MultiBoxBuilder.h"
 #include "Framework/MultiBox/MultiBoxExtender.h"
+#include "IAudioInsightsModule.h"
 #include "IAudioModulation.h"
 #include "ICurveEditorModule.h"
+#include "Insights/Views/ModulationDashboardViewFactory.h"
 #include "Internationalization/Internationalization.h"
 #include "Layouts/SoundControlBusMixStageLayout.h"
 #include "Layouts/SoundControlModulationPatchLayout.h"
@@ -36,7 +38,6 @@
 
 DEFINE_LOG_CATEGORY(LogAudioModulationEditor);
 
-
 namespace AudioModulationEditor
 {
 	static const FName ToolName = TEXT("AssetTools");
@@ -51,6 +52,55 @@ namespace AudioModulationEditor
 	}
 } // namespace AudioModulationEditor
 
+namespace
+{
+
+	static bool bIsModulationRegisteredInAudioInsights = false;
+
+	void RegisterModulationInAudioInsights()
+	{
+		if (!bIsModulationRegisteredInAudioInsights)
+		{
+			IAudioInsightsModule& InsightsModule = FModuleManager::GetModuleChecked<IAudioInsightsModule>(IAudioInsightsModule::GetName());
+			InsightsModule.RegisterDashboardViewFactory(MakeShared<AudioModulationEditor::FAudioModulationDashboardViewFactory>());
+
+			bIsModulationRegisteredInAudioInsights = true;
+		}
+	}
+
+	void UnregisterModulationInAudioInsights()
+	{
+		if (bIsModulationRegisteredInAudioInsights)
+		{
+			IAudioInsightsModule& InsightsModule = FModuleManager::GetModuleChecked<IAudioInsightsModule>(IAudioInsightsModule::GetName());
+			InsightsModule.UnregisterDashboardViewFactory("AudioModulation");
+
+			bIsModulationRegisteredInAudioInsights = false;
+		}
+	}
+
+	void OnModulationInsightsCVarValueChanged(int32 CVarValue)
+	{
+		if (CVarValue)
+		{
+			RegisterModulationInAudioInsights();
+		}
+		else
+		{
+			UnregisterModulationInAudioInsights();
+		}
+	}
+
+}
+
+static int32 AudioModulationInsightsTabCVar = 0;
+static FAutoConsoleVariableRef AudioModulationEnableInsightsTab(
+	TEXT("au.Modulation.AudioInsightsTab"),
+	AudioModulationInsightsTabCVar,
+	TEXT("Enable/disable the Audio Modulation Tab in Audio Insights. 1 = enable, 0 = disable\n")
+	TEXT("Requires closing and Reopening Insights to take effect.\n"),
+	FConsoleVariableDelegate::CreateLambda([](IConsoleVariable* Var) { ::OnModulationInsightsCVarValueChanged(AudioModulationInsightsTabCVar); }),
+	ECVF_Default);
 
 FAudioModulationEditorModule::FAudioModulationEditorModule()
 {
@@ -153,6 +203,12 @@ void FAudioModulationEditorModule::StartupModule()
 			Parameter->RemoveFromRoot();
 		}
 	});
+
+
+	if (AudioModulationInsightsTabCVar)
+	{
+		::RegisterModulationInAudioInsights();
+	}
 }
 
 void FAudioModulationEditorModule::RegisterCustomPropertyLayouts()
@@ -195,6 +251,11 @@ void FAudioModulationEditorModule::ShutdownModule()
 	FModPatchCurveEditorModel::ModPatchViewId = ECurveEditorViewID::Invalid;
 
 	FSlateStyleRegistry::UnRegisterSlateStyle(*StyleSet.Get());
+
+	if (AudioModulationInsightsTabCVar)
+	{
+		::UnregisterModulationInAudioInsights();
+	}
 }
 
 IMPLEMENT_MODULE(FAudioModulationEditorModule, AudioModulationEditor);
diff --git a/Engine/Plugins/Runtime/AudioModulation/Source/AudioModulationEditor/Private/Insights/Views/ModulationDashboardViewFactory.cpp b/Engine/Plugins/Runtime/AudioModulation/Source/AudioModulationEditor/Private/Insights/Views/ModulationDashboardViewFactory.cpp
new file mode 100644
--- /dev/null
+++ b/Engine/Plugins/Runtime/AudioModulation/Source/AudioModulationEditor/Private/Insights/Views/ModulationDashboardViewFactory.cpp
@@ -0,0 +1,91 @@
+// Copyright Epic Games, Inc. All Rights Reserved.
+#include "ModulationDashboardViewFactory.h"
+
+#include "AssetRegistry/AssetData.h"
+#include "AssetRegistry/AssetRegistryModule.h"
+#include "AudioDefines.h"
+#include "AudioDeviceManager.h"
+#include "Editor.h"
+#include "IAudioInsightsModule.h"
+#include "Internationalization/Text.h"
+#include "Subsystems/AssetEditorSubsystem.h"
+#include "Templates/SharedPointer.h"
+
+#define LOCTEXT_NAMESPACE "AudioModulationInsights"
+
+namespace AudioModulationEditor
+{
+
+	FAudioModulationDashboardViewFactory::FAudioModulationDashboardViewFactory()
+	{
+		Providers = TArray<TSharedPtr<UE::Audio::Insights::FTraceProviderBase>>();
+	}
+
+	FAudioModulationDashboardViewFactory::~FAudioModulationDashboardViewFactory()
+	{
+	}
+
+	FName FAudioModulationDashboardViewFactory::GetName() const
+	{
+		return "AudioModulation";
+	}
+
+	FText FAudioModulationDashboardViewFactory::GetDisplayName() const
+	{
+		return LOCTEXT("AudioInsights_Modulation_DisplayName", "Audio Modulation");
+	}
+
+	TSharedRef<SWidget> FAudioModulationDashboardViewFactory::GenerateWidgetForColumn(TSharedRef<UE::Audio::Insights::IDashboardDataViewEntry> InRowData, const FName& InColumnName)
+	{
+		return SNullWidget::NullWidget;
+	}
+
+	void FAudioModulationDashboardViewFactory::ProcessEntries(UE::Audio::Insights::FTraceTableDashboardViewFactory::EProcessReason Reason)
+	{
+	}
+
+	FSlateIcon FAudioModulationDashboardViewFactory::GetIcon() const
+	{
+		return FSlateIcon();
+	}
+
+	UE::Audio::Insights::EDefaultDashboardTabStack FAudioModulationDashboardViewFactory::GetDefaultTabStack() const
+	{
+		return UE::Audio::Insights::EDefaultDashboardTabStack::Analysis;
+	}
+
+	TSharedRef<SWidget> FAudioModulationDashboardViewFactory::MakeWidget()
+	{
+		if (!DashboardWidget.IsValid())
+		{
+			DashboardWidget = UE::Audio::Insights::FTraceTableDashboardViewFactory::MakeWidget();
+
+			if (FilteredEntriesListView.IsValid())
+			{
+				FilteredEntriesListView->SetSelectionMode(ESelectionMode::Single);
+			}
+		}
+
+		return DashboardWidget->AsShared();
+	}
+
+	const TMap<FName, UE::Audio::Insights::FTraceTableDashboardViewFactory::FColumnData>& FAudioModulationDashboardViewFactory::GetColumns() const
+	{
+		static const TMap<FName, UE::Audio::Insights::FTraceTableDashboardViewFactory::FColumnData> ColumnData;
+		
+		return ColumnData;
+	}
+
+	void FAudioModulationDashboardViewFactory::SortTable()
+	{
+		
+	}
+
+	void FAudioModulationDashboardViewFactory::OnSelectionChanged(TSharedPtr<UE::Audio::Insights::IDashboardDataViewEntry> SelectedItem, ESelectInfo::Type SelectInfo)
+	{
+
+	}
+
+} // namespace AudioModulationEditor
+
+#undef LOCTEXT_NAMESPACE
diff --git a/Engine/Plugins/Runtime/AudioModulation/Source/AudioModulationEditor/Private/Insights/Views/ModulationDashboardViewFactory.h b/Engine/Plugins/Runtime/AudioModulation/Source/AudioModulationEditor/Private/Insights/Views/ModulationDashboardViewFactory.h
new file mode 100644
--- /dev/null
+++ b/Engine/Plugins/Runtime/AudioModulation/Source/AudioModulationEditor/Private/Insights/Views/ModulationDashboardViewFactory.h
@@ -0,0 +1,33 @@
+// Copyright Epic Games, Inc. All Rights Reserved.
+#pragma once
+
+#include "Delegates/Delegate.h"
+#include "Templates/SharedPointer.h"
+#include "Views/TableDashboardViewFactory.h"
+#include "Widgets/Input/SCheckBox.h"
+
+namespace AudioModulationEditor
+{
+	class FAudioModulationDashboardViewFactory : public UE::Audio::Insights::FTraceObjectTableDashboardViewFactory
+	{
+	public:
+		FAudioModulationDashboardViewFactory();
+		virtual ~FAudioModulationDashboardViewFactory();
+
+		virtual FName GetName() const override;
+		virtual FText GetDisplayName() const override;
+		virtual FSlateIcon GetIcon() const override;
+		virtual UE::Audio::Insights::EDefaultDashboardTabStack GetDefaultTabStack() const override;
+		virtual TSharedRef<SWidget> MakeWidget() override;
+
+	protected:
+		virtual TSharedRef<SWidget> GenerateWidgetForColumn(TSharedRef<UE::Audio::Insights::IDashboardDataViewEntry> InRowData, const FName& InColumnName) override;
+		virtual void ProcessEntries(UE::Audio::Insights::FTraceTableDashboardViewFactory::EProcessReason Reason) override;
+		virtual const TMap<FName, UE::Audio::Insights::FTraceTableDashboardViewFactory::FColumnData>& GetColumns() const override;
+
+		virtual void SortTable() override;
+
+		virtual void OnSelectionChanged(TSharedPtr<UE::Audio::Insights::IDashboardDataViewEntry> SelectedItem, ESelectInfo::Type SelectInfo) override;
+
+	};
+} // namespace AudioModulationEditor
