Fix bug where DeactivateBusMix was not getting the Bus Mix Unique ID if the bus mix had no stages

diff --git a/Engine/Plugins/Runtime/AudioModulation/Source/AudioModulation/Private/AudioModulationSystem.cpp b/Engine/Plugins/Runtime/AudioModulation/Source/AudioModulation/Private/AudioModulationSystem.cpp
--- a/Engine/Plugins/Runtime/AudioModulation/Source/AudioModulation/Private/AudioModulationSystem.cpp
+++ b/Engine/Plugins/Runtime/AudioModulation/Source/AudioModulation/Private/AudioModulationSystem.cpp
@@ -173,8 +173,7 @@ namespace AudioModulation
 
 	void FAudioModulationSystem::DeactivateBusMix(const USoundControlBusMix& InBusMix)
 	{
-		FBusMixHandle MixHandle = FBusMixHandle::Get(InBusMix.GetUniqueID(), RefProxies.BusMixes);
-		ActiveBusMixIds.Remove(MixHandle.GetId());
+		ActiveBusMixIds.Remove(InBusMix.GetUniqueID());
 		
 		RunCommandOnProcessingThread([this, BusMixId = static_cast<FBusMixId>(InBusMix.GetUniqueID())]()
 		{
