Protect against more nullptrs in MetaSound type hash functions

diff --git a/Engine/Plugins/Runtime/AudioModulation/Source/AudioModulation/Public/SoundModulatorAsset.h b/Engine/Plugins/Runtime/AudioModulation/Source/AudioModulation/Public/SoundModulatorAsset.h
--- a/Engine/Plugins/Runtime/AudioModulation/Source/AudioModulation/Public/SoundModulatorAsset.h
+++ b/Engine/Plugins/Runtime/AudioModulation/Source/AudioModulation/Public/SoundModulatorAsset.h
@@ -72,7 +72,11 @@ namespace AudioModulation
 				
 		friend FORCEINLINE uint32 GetTypeHash(const AudioModulation::FSoundModulatorAsset& InModulatorAsset)
 		{
-			return InModulatorAsset->GetModulatorId();
+			if (InModulatorAsset.IsValid())
+			{
+				return InModulatorAsset->GetModulatorId();
+			}
+			return INDEX_NONE;
 		}
 	};
 
@@ -111,7 +115,11 @@ namespace AudioModulation
 
 		friend FORCEINLINE uint32 GetTypeHash(const AudioModulation::FSoundModulationParameterAsset& InModulationParameterAsset)
 		{
-			return GetTypeHash(InModulationParameterAsset->GetParameter());
+			if (InModulationParameterAsset.IsValid())
+			{
+				return GetTypeHash(InModulationParameterAsset->GetParameter());
+			}
+			return INDEX_NONE;
 		}
 	};
 } // namespace AudioModulation
diff --git a/Engine/Plugins/Runtime/Metasound/Source/MetasoundFrontend/Public/MetasoundParameterPack.h b/Engine/Plugins/Runtime/Metasound/Source/MetasoundFrontend/Public/MetasoundParameterPack.h
--- a/Engine/Plugins/Runtime/Metasound/Source/MetasoundFrontend/Public/MetasoundParameterPack.h
+++ b/Engine/Plugins/Runtime/Metasound/Source/MetasoundFrontend/Public/MetasoundParameterPack.h
@@ -509,7 +509,11 @@ public:
 
 	friend FORCEINLINE uint32 GetTypeHash(const FMetasoundParameterStorageWrapper& InParameterStorageWrapper)
 	{
-		return GetTypeHash(InParameterStorageWrapper->Storage);
+		if (InParameterStorageWrapper.IsPackValid())
+		{
+			return GetTypeHash(InParameterStorageWrapper->Storage);
+		}
+		return INDEX_NONE;
 	}
 };
 
