Control Bus Dashboard

diff --git a/Engine/Plugins/AudioInsights/Source/AudioInsights/Private/AudioInsightsDashboardAssetCommands.cpp b/Engine/Plugins/AudioInsights/Source/AudioInsights/Private/AudioInsightsDashboardAssetCommands.cpp
--- a/Engine/Plugins/AudioInsights/Source/AudioInsights/Private/AudioInsightsDashboardAssetCommands.cpp
+++ b/Engine/Plugins/AudioInsights/Source/AudioInsights/Private/AudioInsightsDashboardAssetCommands.cpp
@@ -70,7 +70,7 @@ namespace UE::Audio::Insights
 
 	FSlateIcon FDashboardAssetCommands::GetStartIcon() const
 	{
-		const FTraceModule& TraceModule = FAudioInsightsModule::GetChecked().GetTraceModule();
+		const IAudioInsightsTraceModule& TraceModule = FAudioInsightsModule::GetChecked().GetTraceModule();
 		return FSlateStyle::Get().CreateIcon(TraceModule.IsTraceAnalysisActive()
 			? "AudioInsights.Icon.Start.Inactive"
 			: "AudioInsights.Icon.Start.Active"
@@ -79,7 +79,7 @@ namespace UE::Audio::Insights
 
 	FSlateIcon FDashboardAssetCommands::GetStopIcon() const
 	{
-		const FTraceModule& TraceModule = FAudioInsightsModule::GetChecked().GetTraceModule();
+		const IAudioInsightsTraceModule& TraceModule = FAudioInsightsModule::GetChecked().GetTraceModule();
 		return FSlateStyle::Get().CreateIcon(TraceModule.IsTraceAnalysisActive()
 			? "AudioInsights.Icon.Stop.Active"
 			: "AudioInsights.Icon.Stop.Inactive"
diff --git a/Engine/Plugins/AudioInsights/Source/AudioInsights/Private/AudioInsightsModule.cpp b/Engine/Plugins/AudioInsights/Source/AudioInsights/Private/AudioInsightsModule.cpp
--- a/Engine/Plugins/AudioInsights/Source/AudioInsights/Private/AudioInsightsModule.cpp
+++ b/Engine/Plugins/AudioInsights/Source/AudioInsights/Private/AudioInsightsModule.cpp
@@ -95,11 +95,6 @@ namespace UE::Audio::Insights
 		return DashboardFactory->GetDeviceId();
 	}
 
-	FAudioInsightsModule& FAudioInsightsModule::GetChecked()
-	{
-		return static_cast<FAudioInsightsModule&>(FModuleManager::GetModuleChecked<IAudioInsightsModule>(GetName()));
-	}
-
 	TSharedRef<FDashboardFactory> FAudioInsightsModule::GetDashboardFactory()
 	{
 		return DashboardFactory->AsShared();
@@ -110,7 +105,7 @@ namespace UE::Audio::Insights
 		return DashboardFactory->AsShared();
 	}
 
-	FTraceModule& FAudioInsightsModule::GetTraceModule()
+	IAudioInsightsTraceModule& FAudioInsightsModule::GetTraceModule()
 	{
 		return TraceModule;
 	}
diff --git a/Engine/Plugins/AudioInsights/Source/AudioInsights/Private/AudioInsightsTraceModule.cpp b/Engine/Plugins/AudioInsights/Source/AudioInsights/Private/AudioInsightsTraceModule.cpp
--- a/Engine/Plugins/AudioInsights/Source/AudioInsights/Private/AudioInsightsTraceModule.cpp
+++ b/Engine/Plugins/AudioInsights/Source/AudioInsights/Private/AudioInsightsTraceModule.cpp
@@ -32,6 +32,11 @@ namespace UE::Audio::Insights
 		OutModuleInfo.DisplayName = TEXT("Audio");
 	}
 
+	void FTraceModule::AddTraceProvider(TSharedPtr<FTraceProviderBase> TraceProvider)
+	{
+		TraceProviders.Add(TraceProvider->GetName(), TraceProvider);
+	}
+
 	TSharedRef<FTraceChannelManager> FTraceModule::GetChannelManager()
 	{
 		return ChannelManager;
diff --git a/Engine/Plugins/AudioInsights/Source/AudioInsights/Private/IAudioInsightsModule.cpp b/Engine/Plugins/AudioInsights/Source/AudioInsights/Private/IAudioInsightsModule.cpp
new file mode 100644
--- /dev/null
+++ b/Engine/Plugins/AudioInsights/Source/AudioInsights/Private/IAudioInsightsModule.cpp
@@ -0,0 +1,15 @@
+// Copyright Epic Games, Inc. All Rights Reserved.
+#include "IAudioInsightsModule.h"
+
+#include "AudioInsightsModule.h"
+
+IAudioInsightsTraceModule& IAudioInsightsModule::GetTraceModule()
+{
+	return GetChecked().GetTraceModule();
+}
+
+UE::Audio::Insights::FAudioInsightsModule& IAudioInsightsModule::GetChecked()
+{
+	using namespace UE::Audio::Insights;
+	return static_cast<FAudioInsightsModule&>(FModuleManager::GetModuleChecked<IAudioInsightsModule>(FAudioInsightsModule::GetName()));
+}
diff --git a/Engine/Plugins/AudioInsights/Source/AudioInsights/Private/Providers/MixerSourceTraceProvider.cpp b/Engine/Plugins/AudioInsights/Source/AudioInsights/Private/Providers/MixerSourceTraceProvider.cpp
--- a/Engine/Plugins/AudioInsights/Source/AudioInsights/Private/Providers/MixerSourceTraceProvider.cpp
+++ b/Engine/Plugins/AudioInsights/Source/AudioInsights/Private/Providers/MixerSourceTraceProvider.cpp
@@ -169,7 +169,7 @@ namespace UE::Audio::Insights
 
 				if (!bIsFirstTimestampSet)
 				{
-					FTraceModule& TraceModule = FAudioInsightsModule::GetChecked().GetTraceModule();
+					FTraceModule& TraceModule = static_cast<FTraceModule&>(FAudioInsightsModule::GetChecked().GetTraceModule());
 
 					if (TraceModule.GetFirstTimeStamp() < 0.0)
 					{
diff --git a/Engine/Plugins/AudioInsights/Source/AudioInsights/Private/Views/MixerSourceDashboardViewFactory.cpp b/Engine/Plugins/AudioInsights/Source/AudioInsights/Private/Views/MixerSourceDashboardViewFactory.cpp
--- a/Engine/Plugins/AudioInsights/Source/AudioInsights/Private/Views/MixerSourceDashboardViewFactory.cpp
+++ b/Engine/Plugins/AudioInsights/Source/AudioInsights/Private/Views/MixerSourceDashboardViewFactory.cpp
@@ -51,7 +51,7 @@ namespace UE::Audio::Insights
 
 	FMixerSourceDashboardViewFactory::FMixerSourceDashboardViewFactory()
 	{
-		const FTraceModule& TraceModule = FAudioInsightsModule::GetChecked().GetTraceModule();
+		const FTraceModule& TraceModule = static_cast<FTraceModule&>(FAudioInsightsModule::GetChecked().GetTraceModule());
 		Providers = TArray<TSharedPtr<FTraceProviderBase>>
 		{
 			TraceModule.FindAudioTraceProvider<FMixerSourceTraceProvider>()
@@ -654,7 +654,7 @@ namespace UE::Audio::Insights
 
 			if (GameState == EGameState::Running)
 			{
-				const FTraceModule& TraceModule = FAudioInsightsModule::GetChecked().GetTraceModule();
+				const FTraceModule& TraceModule = static_cast<FTraceModule&>(FAudioInsightsModule::GetChecked().GetTraceModule());
 				const double FirstTimestamp = TraceModule.GetFirstTimeStamp();
 
 				const double CurrentTime = FPlatformTime::Seconds();
diff --git a/Engine/Plugins/AudioInsights/Source/AudioInsights/Private/Views/VirtualLoopDashboardViewFactory.cpp b/Engine/Plugins/AudioInsights/Source/AudioInsights/Private/Views/VirtualLoopDashboardViewFactory.cpp
--- a/Engine/Plugins/AudioInsights/Source/AudioInsights/Private/Views/VirtualLoopDashboardViewFactory.cpp
+++ b/Engine/Plugins/AudioInsights/Source/AudioInsights/Private/Views/VirtualLoopDashboardViewFactory.cpp
@@ -28,7 +28,7 @@ namespace UE::Audio::Insights
 
 	FVirtualLoopDashboardViewFactory::FVirtualLoopDashboardViewFactory()
 	{
-		const FTraceModule& TraceModule = FAudioInsightsModule::GetChecked().GetTraceModule();
+		const FTraceModule& TraceModule = static_cast<FTraceModule&>(FAudioInsightsModule::GetChecked().GetTraceModule());
 		Providers = TArray<TSharedPtr<FTraceProviderBase>>
 		{
 			TraceModule.FindAudioTraceProvider<FVirtualLoopTraceProvider>()
diff --git a/Engine/Plugins/AudioInsights/Source/AudioInsights/Public/AudioInsightsModule.h b/Engine/Plugins/AudioInsights/Source/AudioInsights/Public/AudioInsightsModule.h
--- a/Engine/Plugins/AudioInsights/Source/AudioInsights/Public/AudioInsightsModule.h
+++ b/Engine/Plugins/AudioInsights/Source/AudioInsights/Public/AudioInsightsModule.h
@@ -26,8 +26,7 @@ namespace UE::Audio::Insights
 		virtual void UnregisterDashboardViewFactory(FName InName) override;
 		virtual ::Audio::FDeviceId GetDeviceId() const override;
 
-		AUDIOINSIGHTS_API static FAudioInsightsModule& GetChecked();
-		AUDIOINSIGHTS_API FTraceModule& GetTraceModule();
+		AUDIOINSIGHTS_API virtual IAudioInsightsTraceModule& GetTraceModule() override;
 
 		TSharedRef<FDashboardFactory> GetDashboardFactory();
 		const TSharedRef<FDashboardFactory> GetDashboardFactory() const;
diff --git a/Engine/Plugins/AudioInsights/Source/AudioInsights/Public/AudioInsightsTraceModule.h b/Engine/Plugins/AudioInsights/Source/AudioInsights/Public/AudioInsightsTraceModule.h
--- a/Engine/Plugins/AudioInsights/Source/AudioInsights/Public/AudioInsightsTraceModule.h
+++ b/Engine/Plugins/AudioInsights/Source/AudioInsights/Public/AudioInsightsTraceModule.h
@@ -1,17 +1,15 @@
 // Copyright Epic Games, Inc. All Rights Reserved.
 #pragma once
 
-#include "AudioInsightsTraceChannelHandle.h"
-#include "AudioInsightsTraceProviderBase.h"
 #include "Containers/Map.h"
+#include "IAudioInsightsTraceModule.h"
 #include "Templates/SharedPointer.h"
 #include "TraceServices/Model/AnalysisSession.h"
-#include "TraceServices/ModuleService.h"
 #include "UObject/NameTypes.h"
 
 namespace UE::Audio::Insights
 {
-	class AUDIOINSIGHTS_API FTraceModule : public TraceServices::IModule
+	class AUDIOINSIGHTS_API FTraceModule : public IAudioInsightsTraceModule
 	{
 	public:
 		FTraceModule();
@@ -31,11 +29,13 @@ namespace UE::Audio::Insights
 			return StaticCastSharedPtr<TraceProviderType>(TraceProviders.FindRef(TraceProviderType::GetName_Static()));
 		}
 
-		TSharedRef<FTraceChannelManager> GetChannelManager();
+		virtual void AddTraceProvider(TSharedPtr<FTraceProviderBase> TraceProvider) override;
 
-		void StartTraceAnalysis() const;
-		bool IsTraceAnalysisActive() const;
-		void StopTraceAnalysis() const;
+		virtual TSharedRef<FTraceChannelManager> GetChannelManager() override;
+
+		void StartTraceAnalysis() const override;
+		bool IsTraceAnalysisActive() const override;
+		void StopTraceAnalysis() const override;
 
 		void SetFirstTimeStamp(double InFirstTimeStamp) { FirstTimeStamp = InFirstTimeStamp; }
 		double GetFirstTimeStamp() const { return FirstTimeStamp; }
diff --git a/Engine/Plugins/AudioInsights/Source/AudioInsights/Public/AudioInsightsTraceProviderBase.h b/Engine/Plugins/AudioInsights/Source/AudioInsights/Public/AudioInsightsTraceProviderBase.h
--- a/Engine/Plugins/AudioInsights/Source/AudioInsights/Public/AudioInsightsTraceProviderBase.h
+++ b/Engine/Plugins/AudioInsights/Source/AudioInsights/Public/AudioInsightsTraceProviderBase.h
@@ -5,6 +5,7 @@
 #include "AudioDefines.h"
 #include "Containers/Ticker.h"
 #include "IAudioInsightsModule.h"
+#include "IAudioInsightsTraceModule.h"
 #include "Messages/AnalyzerMessageQueue.h"
 #include "Modules/ModuleManager.h"
 #include "Templates/SharedPointer.h"
diff --git a/Engine/Plugins/AudioInsights/Source/AudioInsights/Public/IAudioInsightsModule.h b/Engine/Plugins/AudioInsights/Source/AudioInsights/Public/IAudioInsightsModule.h
--- a/Engine/Plugins/AudioInsights/Source/AudioInsights/Public/IAudioInsightsModule.h
+++ b/Engine/Plugins/AudioInsights/Source/AudioInsights/Public/IAudioInsightsModule.h
@@ -6,9 +6,13 @@
 #include "Templates/SharedPointer.h"
 #include "UObject/NameTypes.h"
 
+class IAudioInsightsTraceModule;
+
 namespace UE::Audio::Insights
 {
 	class IDashboardViewFactory;
+	class FAudioInsightsModule;
+	class FTraceProviderBase;
 } // namespace UE::Audio::Insights
 
 class AUDIOINSIGHTS_API IAudioInsightsModule : public IModuleInterface
@@ -19,6 +23,10 @@ public:
 
 	virtual ::Audio::FDeviceId GetDeviceId() const = 0;
 
+	virtual IAudioInsightsTraceModule& GetTraceModule();
+
+	static UE::Audio::Insights::FAudioInsightsModule& GetChecked();
+
 	static FName GetName()
 	{
 		const FLazyName ModuleName = "AudioInsights";
diff --git a/Engine/Plugins/AudioInsights/Source/AudioInsights/Public/IAudioInsightsTraceModule.h b/Engine/Plugins/AudioInsights/Source/AudioInsights/Public/IAudioInsightsTraceModule.h
new file mode 100644
--- /dev/null
+++ b/Engine/Plugins/AudioInsights/Source/AudioInsights/Public/IAudioInsightsTraceModule.h
@@ -0,0 +1,23 @@
+// Copyright Epic Games, Inc. All Rights Reserved.
+#pragma once
+
+#include "AudioInsightsTraceChannelHandle.h"
+#include "AudioInsightsTraceProviderBase.h"
+#include "TraceServices/ModuleService.h"
+
+namespace UE::Audio::Insights
+{
+	class FTraceProviderBase;
+} // namespace UE::Audio::Insights
+
+class AUDIOINSIGHTS_API IAudioInsightsTraceModule : public TraceServices::IModule
+{
+public:
+	virtual void AddTraceProvider(TSharedPtr<UE::Audio::Insights::FTraceProviderBase> TraceProvider) = 0;
+
+	virtual void StartTraceAnalysis() const = 0;
+	virtual bool IsTraceAnalysisActive() const = 0;
+	virtual void StopTraceAnalysis() const = 0;
+
+	virtual TSharedRef<UE::Audio::Insights::FTraceChannelManager> GetChannelManager() = 0;
+};
diff --git a/Engine/Plugins/AudioInsights/Source/AudioInsightsEditor/Private/AudioInsightsEditorDashboardFactory.cpp b/Engine/Plugins/AudioInsights/Source/AudioInsightsEditor/Private/AudioInsightsEditorDashboardFactory.cpp
--- a/Engine/Plugins/AudioInsights/Source/AudioInsightsEditor/Private/AudioInsightsEditorDashboardFactory.cpp
+++ b/Engine/Plugins/AudioInsights/Source/AudioInsightsEditor/Private/AudioInsightsEditorDashboardFactory.cpp
@@ -61,7 +61,7 @@ namespace UE::Audio::Insights
 		{
 			if (bStartWithPIE)
 			{
-				const FTraceModule& TraceModule = FAudioInsightsModule::GetChecked().GetTraceModule();
+				const IAudioInsightsTraceModule& TraceModule = FAudioInsightsModule::GetChecked().GetTraceModule();
 				TraceModule.StartTraceAnalysis();
 
 				ActiveDeviceId = InDeviceId;
@@ -75,7 +75,7 @@ namespace UE::Audio::Insights
 	{
 		if (bStartWithPIE)
 		{
-			const FTraceModule& TraceModule = FAudioInsightsModule::GetChecked().GetTraceModule();
+			const IAudioInsightsTraceModule& TraceModule = FAudioInsightsModule::GetChecked().GetTraceModule();
 			TraceModule.StartTraceAnalysis();
 		}
 	}
@@ -89,7 +89,7 @@ namespace UE::Audio::Insights
 	{
 		if (bStopWithPIE)
 		{
-			const FTraceModule& TraceModule = FAudioInsightsModule::GetChecked().GetTraceModule();
+			const IAudioInsightsTraceModule& TraceModule = FAudioInsightsModule::GetChecked().GetTraceModule();
 			TraceModule.StopTraceAnalysis();
 		}
 
diff --git a/Engine/Plugins/AudioInsights/Source/AudioInsightsEditor/Private/AudioInsightsEditorModule.cpp b/Engine/Plugins/AudioInsights/Source/AudioInsightsEditor/Private/AudioInsightsEditorModule.cpp
--- a/Engine/Plugins/AudioInsights/Source/AudioInsightsEditor/Private/AudioInsightsEditorModule.cpp
+++ b/Engine/Plugins/AudioInsights/Source/AudioInsightsEditor/Private/AudioInsightsEditorModule.cpp
@@ -74,6 +74,13 @@ namespace UE::Audio::Insights
 		return static_cast<FAudioInsightsEditorModule&>(FModuleManager::GetModuleChecked<IAudioInsightsModule>(GetEditorName()));
 	}
 
+	IAudioInsightsTraceModule& FAudioInsightsEditorModule::GetTraceModule()
+	{
+		IAudioInsightsModule& InsightsModule = FModuleManager::LoadModuleChecked<IAudioInsightsModule>(IAudioInsightsModule::GetName());
+		return InsightsModule.GetTraceModule();
+
+	}
+
 	TSharedRef<FEditorDashboardFactory> FAudioInsightsEditorModule::GetDashboardFactory()
 	{
 		return DashboardFactory->AsShared();
diff --git a/Engine/Plugins/AudioInsights/Source/AudioInsightsEditor/Private/AudioInsightsEditorModule.h b/Engine/Plugins/AudioInsights/Source/AudioInsightsEditor/Private/AudioInsightsEditorModule.h
--- a/Engine/Plugins/AudioInsights/Source/AudioInsightsEditor/Private/AudioInsightsEditorModule.h
+++ b/Engine/Plugins/AudioInsights/Source/AudioInsightsEditor/Private/AudioInsightsEditorModule.h
@@ -30,6 +30,7 @@ namespace UE::Audio::Insights
 		const TSharedRef<FEditorDashboardFactory> GetDashboardFactory() const;
 
 		static FAudioInsightsEditorModule& GetChecked();
+		IAudioInsightsTraceModule& GetTraceModule() override;
 
 	private:
 		void RegisterMenus();
diff --git a/Engine/Plugins/Runtime/AudioModulation/Source/AudioModulation/AudioModulation.Build.cs b/Engine/Plugins/Runtime/AudioModulation/Source/AudioModulation/AudioModulation.Build.cs
--- a/Engine/Plugins/Runtime/AudioModulation/Source/AudioModulation/AudioModulation.Build.cs
+++ b/Engine/Plugins/Runtime/AudioModulation/Source/AudioModulation/AudioModulation.Build.cs
@@ -25,10 +25,12 @@ namespace UnrealBuildTool.Rules
 				new string[]
 				{
 					"AudioMixer",
+					"AudioMixerCore",
 					"CoreUObject",
 					"DeveloperSettings",
 					"Engine",
-					"SignalProcessing"
+					"SignalProcessing",
+					"TraceLog"
 				}
 			);
 
diff --git a/Engine/Plugins/Runtime/AudioModulation/Source/AudioModulation/Private/AudioModulationSystem.cpp b/Engine/Plugins/Runtime/AudioModulation/Source/AudioModulation/Private/AudioModulationSystem.cpp
--- a/Engine/Plugins/Runtime/AudioModulation/Source/AudioModulation/Private/AudioModulationSystem.cpp
+++ b/Engine/Plugins/Runtime/AudioModulation/Source/AudioModulation/Private/AudioModulationSystem.cpp
@@ -10,6 +10,7 @@
 
 #include "Async/Async.h"
 #include "Audio/AudioAddressPattern.h"
+#include "AudioMixerTrace.h"
 #include "AudioModulationLogging.h"
 #include "AudioModulationProfileSerializer.h"
 #include "AudioModulationSettings.h"
@@ -35,6 +36,21 @@ DECLARE_DWORD_COUNTER_STAT(TEXT("Mix Count"),	STAT_AudioModulationMixCount, STAT
 DECLARE_DWORD_COUNTER_STAT(TEXT("Patch Count"), STAT_AudioModulationPatchCount, STATGROUP_AudioModulation)
 DECLARE_DWORD_COUNTER_STAT(TEXT("Render Queue Commands Processed"), STAT_AudioModulationProcQueueCount, STATGROUP_AudioModulation)
 
+#if UE_AUDIO_PROFILERTRACE_ENABLED
+UE_TRACE_EVENT_BEGIN(Audio, ControlBusUpdate)
+	UE_TRACE_EVENT_FIELD(uint32, DeviceId)
+	UE_TRACE_EVENT_FIELD(uint32, ControlBusId)
+	UE_TRACE_EVENT_FIELD(double, Timestamp)
+	UE_TRACE_EVENT_FIELD(float, Value)
+UE_TRACE_EVENT_END()
+
+UE_TRACE_EVENT_BEGIN(Audio, ControlBusDeactivate)
+	UE_TRACE_EVENT_FIELD(uint32, DeviceId)
+	UE_TRACE_EVENT_FIELD(uint32, ControlBusId)
+	UE_TRACE_EVENT_FIELD(double, Timestamp)
+UE_TRACE_EVENT_END()
+#endif // UE_AUDIO_PROFILERTRACE_ENABLED
+
 namespace AudioModulation
 {
 	enum class EModulatorType : Audio::FModulatorTypeId
@@ -547,6 +563,21 @@ namespace AudioModulation
 		SET_DWORD_STAT(STAT_AudioModulationPatchCount, RefProxies.Patches.Num());
 		SET_DWORD_STAT(STAT_AudioModulationProcQueueCount, CommandsProcessed);
 
+#if UE_AUDIO_PROFILERTRACE_ENABLED
+		const bool bChannelEnabled = UE_TRACE_CHANNELEXPR_IS_ENABLED(AudioChannel);
+		if (bChannelEnabled)
+		{
+			for (TPair<FBusId, FControlBusProxy>& Pair : RefProxies.Buses)
+			{
+				UE_TRACE_LOG(Audio, ControlBusUpdate, AudioChannel)
+					<< ControlBusUpdate.DeviceId(static_cast<uint32>(AudioDeviceId))
+					<< ControlBusUpdate.ControlBusId(static_cast<uint32>(Pair.Key))
+					<< ControlBusUpdate.Timestamp(FPlatformTime::Cycles64())
+					<< ControlBusUpdate.Value(Pair.Value.GetValue());
+			}
+		}
+#endif // UE_AUDIO_PROFILERTRACE_ENABLED
+
 #if !UE_BUILD_SHIPPING
  		Debugger->UpdateDebugData(InElapsed, RefProxies);
 #endif // !UE_BUILD_SHIPPING
@@ -754,6 +785,13 @@ namespace AudioModulation
 			FBusHandle BusHandle = FBusHandle::Get(static_cast<FBusId>(ModId), RefProxies.Buses);
 			if (UnregisterModulator<FBusHandle>(BusHandle, RefModulators.BusMap, HandleId))
 			{
+#if UE_AUDIO_PROFILERTRACE_ENABLED
+				UE_TRACE_LOG(Audio, ControlBusDeactivate, AudioChannel)
+					<< ControlBusDeactivate.DeviceId(static_cast<uint32>(AudioDeviceId))
+					<< ControlBusDeactivate.ControlBusId(static_cast<uint32>(ModId))
+					<< ControlBusDeactivate.Timestamp(FPlatformTime::Cycles64());
+#endif
+
 				return;
 			}
 
diff --git a/Engine/Plugins/Runtime/AudioModulation/Source/AudioModulation/Private/SoundControlBusProxy.cpp b/Engine/Plugins/Runtime/AudioModulation/Source/AudioModulation/Private/SoundControlBusProxy.cpp
--- a/Engine/Plugins/Runtime/AudioModulation/Source/AudioModulation/Private/SoundControlBusProxy.cpp
+++ b/Engine/Plugins/Runtime/AudioModulation/Source/AudioModulation/Private/SoundControlBusProxy.cpp
@@ -3,12 +3,22 @@
 
 #include "AudioDevice.h"
 #include "AudioDeviceManager.h"
+#include "AudioMixerTrace.h"
 #include "AudioModulation.h"
 #include "AudioModulationLogging.h"
 #include "AudioModulationSystem.h"
 #include "Engine/World.h"
 #include "SoundModulationGenerator.h"
 
+#if UE_AUDIO_PROFILERTRACE_ENABLED
+UE_TRACE_EVENT_BEGIN(Audio, ControlBusActivate)
+	UE_TRACE_EVENT_FIELD(uint32, DeviceId)
+	UE_TRACE_EVENT_FIELD(uint32, ControlBusId)
+	UE_TRACE_EVENT_FIELD(double, Timestamp)
+	UE_TRACE_EVENT_FIELD(UE::Trace::WideString, Name)
+	UE_TRACE_EVENT_FIELD(UE::Trace::WideString, ParamName)
+UE_TRACE_EVENT_END()
+#endif // UE_AUDIO_PROFILERTRACE_ENABLED
 
 namespace AudioModulation
 {
@@ -17,6 +27,16 @@ namespace AudioModulation
 	Audio::FModulatorTypeId FControlBusSettings::Register(Audio::FModulatorHandleId HandleId, IAudioModulationManager& InModulation) const
 	{
 		FAudioModulationSystem& ModSystem = static_cast<FAudioModulationManager&>(InModulation).GetSystem();
+
+#if UE_AUDIO_PROFILERTRACE_ENABLED
+		UE_TRACE_LOG(Audio, ControlBusActivate, AudioChannel)
+			<< ControlBusActivate.DeviceId(static_cast<uint32>(ModSystem.GetAudioDeviceId()))
+			<< ControlBusActivate.ControlBusId(static_cast<uint32>(GetId()))
+			<< ControlBusActivate.Timestamp(FPlatformTime::Cycles64())
+			<< ControlBusActivate.Name(*GetName().ToString())
+			<< ControlBusActivate.ParamName(*OutputParameter.ParameterName.ToString());
+#endif
+
 		return ModSystem.RegisterModulator(HandleId, *this);
 	}
 
diff --git a/Engine/Plugins/Runtime/AudioModulation/Source/AudioModulationEditor/AudioModulationEditor.build.cs b/Engine/Plugins/Runtime/AudioModulation/Source/AudioModulationEditor/AudioModulationEditor.build.cs
--- a/Engine/Plugins/Runtime/AudioModulation/Source/AudioModulationEditor/AudioModulationEditor.build.cs
+++ b/Engine/Plugins/Runtime/AudioModulation/Source/AudioModulationEditor/AudioModulationEditor.build.cs
@@ -29,10 +29,12 @@ public class AudioModulationEditor : ModuleRules
 				"AudioEditor",
 				"AudioExtensions",
 				"AudioInsights",
+				"AudioInsightsEditor",
 				"AudioModulation",
 				"CurveEditor",
 				"EditorWidgets",
 				"ToolWidgets",
+				"TraceAnalysis",
 				"WaveTable",
 				"WaveTableEditor"
 			}
diff --git a/Engine/Plugins/Runtime/AudioModulation/Source/AudioModulationEditor/Private/AudioModulationEditor.cpp b/Engine/Plugins/Runtime/AudioModulation/Source/AudioModulationEditor/Private/AudioModulationEditor.cpp
--- a/Engine/Plugins/Runtime/AudioModulation/Source/AudioModulationEditor/Private/AudioModulationEditor.cpp
+++ b/Engine/Plugins/Runtime/AudioModulation/Source/AudioModulationEditor/Private/AudioModulationEditor.cpp
@@ -16,7 +16,7 @@
 #include "IAudioInsightsModule.h"
 #include "IAudioModulation.h"
 #include "ICurveEditorModule.h"
-#include "Insights/Views/ModulationDashboardViewFactory.h"
+#include "Insights/Views/ControlBusDashboardViewFactory.h"
 #include "Internationalization/Internationalization.h"
 #include "Layouts/SoundControlBusMixStageLayout.h"
 #include "Layouts/SoundControlModulationPatchLayout.h"
@@ -54,15 +54,15 @@ namespace AudioModulationEditor
 
 namespace
 {
-
 	static bool bIsModulationRegisteredInAudioInsights = false;
 
 	void RegisterModulationInAudioInsights()
 	{
 		if (!bIsModulationRegisteredInAudioInsights)
 		{
-			IAudioInsightsModule& InsightsModule = FModuleManager::GetModuleChecked<IAudioInsightsModule>(IAudioInsightsModule::GetName());
-			InsightsModule.RegisterDashboardViewFactory(MakeShared<AudioModulationEditor::FAudioModulationDashboardViewFactory>());
+
+			IAudioInsightsModule& InsightsModule = FModuleManager::LoadModuleChecked<IAudioInsightsModule>(IAudioInsightsModule::GetEditorName());
+			InsightsModule.RegisterDashboardViewFactory(MakeShared<AudioModulationEditor::FControlBusDashboardViewFactory>());
 
 			bIsModulationRegisteredInAudioInsights = true;
 		}
@@ -72,9 +72,8 @@ namespace
 	{
 		if (bIsModulationRegisteredInAudioInsights)
 		{
-			IAudioInsightsModule& InsightsModule = FModuleManager::GetModuleChecked<IAudioInsightsModule>(IAudioInsightsModule::GetName());
-			InsightsModule.UnregisterDashboardViewFactory("AudioModulation");
-
+			IAudioInsightsModule& InsightsModule = FModuleManager::GetModuleChecked<IAudioInsightsModule>(IAudioInsightsModule::GetEditorName());
+			InsightsModule.UnregisterDashboardViewFactory("ControlBuses");
 			bIsModulationRegisteredInAudioInsights = false;
 		}
 	}
diff --git a/Engine/Plugins/Runtime/AudioModulation/Source/AudioModulationEditor/Private/Insights/Messages/ControlBusTraceMessages.h b/Engine/Plugins/Runtime/AudioModulation/Source/AudioModulationEditor/Private/Insights/Messages/ControlBusTraceMessages.h
new file mode 100644
--- /dev/null
+++ b/Engine/Plugins/Runtime/AudioModulation/Source/AudioModulationEditor/Private/Insights/Messages/ControlBusTraceMessages.h
@@ -0,0 +1,83 @@
+// Copyright Epic Games, Inc. All Rights Reserved.
+#pragma once
+
+#include "AudioDefines.h"
+#include "Math/NumericLimits.h"
+#include "Messages/AnalyzerMessageQueue.h"
+#include "Trace/Analyzer.h"
+#include "Views/TableDashboardViewFactory.h"
+
+namespace AudioModulationEditor
+{
+
+	using FBusId = uint32;
+
+	struct FControlBusMessageBase
+	{
+		FControlBusMessageBase() = default;
+		FControlBusMessageBase(const UE::Trace::IAnalyzer::FOnEventContext& InContext)
+		{
+			const UE::Trace::IAnalyzer::FEventData& EventData = InContext.EventData;
+			DeviceId = static_cast<::Audio::FDeviceId>(EventData.GetValue<uint32>("DeviceId"));
+			Timestamp = InContext.EventTime.AsSeconds(EventData.GetValue<uint64>("Timestamp"));
+			ControlBusId = static_cast<FBusId>(EventData.GetValue<uint32>("ControlBusId"));
+		}
+
+		Audio::FDeviceId DeviceId = INDEX_NONE;
+		FBusId ControlBusId;
+		double Timestamp = 0.0;
+	};
+
+	struct FControlBusActivateMessage : public FControlBusMessageBase
+	{
+		FControlBusActivateMessage() = default;
+		FControlBusActivateMessage(const UE::Trace::IAnalyzer::FOnEventContext& InContext)
+			: FControlBusMessageBase(InContext)
+		{
+			const UE::Trace::IAnalyzer::FEventData& EventData = InContext.EventData;
+			EventData.GetString("Name", BusName);
+			EventData.GetString("ParamName", ParamName);
+		}
+
+		FString BusName;
+		FString ParamName;
+	};
+
+	using FControlBusDeactivateMessage = FControlBusMessageBase;
+
+	struct FControlBusUpdateMessage : public FControlBusMessageBase
+	{
+		FControlBusUpdateMessage() = default;
+		FControlBusUpdateMessage(const UE::Trace::IAnalyzer::FOnEventContext& InContext)
+			: FControlBusMessageBase(InContext)
+		{
+			const UE::Trace::IAnalyzer::FEventData& EventData = InContext.EventData;
+			Value = EventData.GetValue<float>("Value");
+		}
+
+		float Value = 1.0f;
+	};
+
+	class FControlBusDashboardEntry : public UE::Audio::Insights::FSoundAssetDashboardEntry
+	{
+	public:
+		FControlBusDashboardEntry() = default;
+		virtual ~FControlBusDashboardEntry() = default;
+
+		FText GetParamNameAsFText() const { return FText::FromString(ParamName); }
+
+		FBusId ControlBusId = INDEX_NONE;
+		float Value = 1.0f;
+		FString ParamName;
+	};
+
+	class FControlBusMessages
+	{
+		UE::Audio::Insights::TAnalyzerMessageQueue<FControlBusUpdateMessage> UpdateMessages{ 2.0 };
+		UE::Audio::Insights::TAnalyzerMessageQueue<FControlBusActivateMessage> ActivateMessages{ 0.1 };
+		UE::Audio::Insights::TAnalyzerMessageQueue<FControlBusDeactivateMessage> DeactivateMessages{ 0.1 };
+
+		friend class FControlBusTraceProvider;
+	};
+
+} // namespace AudioModulationEditor
diff --git a/Engine/Plugins/Runtime/AudioModulation/Source/AudioModulationEditor/Private/Insights/Providers/ControlBusTraceProvider.cpp b/Engine/Plugins/Runtime/AudioModulation/Source/AudioModulationEditor/Private/Insights/Providers/ControlBusTraceProvider.cpp
new file mode 100644
--- /dev/null
+++ b/Engine/Plugins/Runtime/AudioModulation/Source/AudioModulationEditor/Private/Insights/Providers/ControlBusTraceProvider.cpp
@@ -0,0 +1,133 @@
+// Copyright Epic Games, Inc. All Rights Reserved.
+#include "Insights/Providers/ControlBusTraceProvider.h"
+
+#include "Trace/Analyzer.h"
+#include "TraceServices/Model/AnalysisSession.h"
+#include "TraceServices/ModuleService.h"
+
+
+namespace AudioModulationEditor
+{
+	FName FControlBusTraceProvider::GetName_Static()
+	{
+		return "ControlBusProvider";
+	}
+
+	bool FControlBusTraceProvider::ProcessMessages()
+	{
+		auto BumpEntryFunc = [this](const FControlBusMessageBase& Msg)
+		{
+			TSharedPtr<FControlBusDashboardEntry>* ToReturn = nullptr;
+			UpdateDeviceEntry(Msg.DeviceId, Msg.ControlBusId, [&ToReturn, &Msg](TSharedPtr<FControlBusDashboardEntry>& Entry)
+				{
+					if (!Entry.IsValid())
+					{
+						Entry = MakeShared<FControlBusDashboardEntry>();
+						Entry->DeviceId = Msg.DeviceId;
+						Entry->ControlBusId = Msg.ControlBusId;
+					}
+					ToReturn = &Entry;
+				});
+
+			return ToReturn;
+		};
+
+		ProcessMessageQueue<FControlBusActivateMessage>(TraceMessages.ActivateMessages, BumpEntryFunc,
+		[](const FControlBusActivateMessage& Msg, TSharedPtr<FControlBusDashboardEntry>* OutEntry)
+		{
+			FControlBusDashboardEntry& EntryRef = *OutEntry->Get();
+			EntryRef.Name = *Msg.BusName;
+			EntryRef.ControlBusId = Msg.ControlBusId;
+			EntryRef.ParamName = Msg.ParamName;
+		});
+
+		ProcessMessageQueue<FControlBusUpdateMessage>(TraceMessages.UpdateMessages, BumpEntryFunc,
+		[](const FControlBusUpdateMessage& Msg, TSharedPtr<FControlBusDashboardEntry>* OutEntry)
+		{
+			FControlBusDashboardEntry& EntryRef = *OutEntry->Get();
+			EntryRef.Value = Msg.Value;
+		});
+
+		auto GetEntry = [this](const FControlBusMessageBase& Msg)
+		{
+			return FindDeviceEntry(Msg.DeviceId, Msg.ControlBusId);
+		};
+
+		ProcessMessageQueue<FControlBusDeactivateMessage>(TraceMessages.DeactivateMessages, GetEntry,
+		[this](const FControlBusDeactivateMessage& Msg, TSharedPtr<FControlBusDashboardEntry>* OutEntry)
+		{
+			if (OutEntry && (*OutEntry)->Timestamp < Msg.Timestamp)
+			{
+				RemoveDeviceEntry(Msg.DeviceId, Msg.ControlBusId);
+			}
+		});
+
+		return true;
+	}
+
+	UE::Trace::IAnalyzer* FControlBusTraceProvider::ConstructAnalyzer()
+	{
+		class FControlBusTraceAnalyzer : public UE::Audio::Insights::FTraceProviderBase::FTraceAnalyzerBase
+		{
+		public:
+			FControlBusTraceAnalyzer(TSharedRef<FControlBusTraceProvider> InProvider)
+				: UE::Audio::Insights::FTraceProviderBase::FTraceAnalyzerBase(InProvider)
+			{
+			}
+
+			virtual void OnAnalysisBegin(const UE::Trace::IAnalyzer::FOnAnalysisContext& Context) override
+			{
+				UE::Audio::Insights::FTraceProviderBase::FTraceAnalyzerBase::OnAnalysisBegin(Context);
+
+				UE::Trace::IAnalyzer::FInterfaceBuilder& Builder = Context.InterfaceBuilder;
+				Builder.RouteEvent(RouteId_Activate, "Audio", "ControlBusActivate");
+				Builder.RouteEvent(RouteId_Deactivate, "Audio", "ControlBusDeactivate");
+				Builder.RouteEvent(RouteId_Update, "Audio", "ControlBusUpdate");
+			}
+
+			virtual bool OnEvent(uint16 RouteId, UE::Trace::IAnalyzer::EStyle Style, const UE::Trace::IAnalyzer::FOnEventContext& Context) override
+			{
+				LLM_SCOPE_BYNAME(TEXT("Insights/FControlBusTraceAnalyzer"));
+
+				FControlBusMessages& Messages = GetProvider<FControlBusTraceProvider>().TraceMessages;
+				switch (RouteId)
+				{
+					case RouteId_Activate:
+					{
+						Messages.ActivateMessages.Enqueue(FControlBusActivateMessage { Context });
+						break;
+					}
+
+					case RouteId_Deactivate:
+					{
+						Messages.DeactivateMessages.Enqueue(FControlBusDeactivateMessage{ Context });
+						break;
+					}
+
+					case RouteId_Update:
+					{
+						Messages.UpdateMessages.Enqueue(FControlBusUpdateMessage { Context });
+						break;
+					}
+
+					default:
+					{
+						return OnEventFailure(RouteId, Style, Context);
+					}
+				}
+
+				return OnEventSuccess(RouteId, Style, Context);
+			}
+
+		private:
+			enum : uint16
+			{
+				RouteId_Activate,
+				RouteId_Deactivate,
+				RouteId_Update
+			};
+		};
+
+		return new FControlBusTraceAnalyzer(AsShared());
+	}
+} // namespace AudioModulationEditor
diff --git a/Engine/Plugins/Runtime/AudioModulation/Source/AudioModulationEditor/Private/Insights/Providers/ControlBusTraceProvider.h b/Engine/Plugins/Runtime/AudioModulation/Source/AudioModulationEditor/Private/Insights/Providers/ControlBusTraceProvider.h
new file mode 100644
--- /dev/null
+++ b/Engine/Plugins/Runtime/AudioModulation/Source/AudioModulationEditor/Private/Insights/Providers/ControlBusTraceProvider.h
@@ -0,0 +1,33 @@
+// Copyright Epic Games, Inc. All Rights Reserved.
+#pragma once
+
+#include "AudioInsightsTraceChannelHandle.h"
+#include "AudioInsightsTraceProviderBase.h"
+#include "Insights/Messages/ControlBusTraceMessages.h"
+
+
+namespace AudioModulationEditor
+{
+	class FControlBusTraceProvider
+		: public UE::Audio::Insights::TDeviceDataMapTraceProvider<uint32, TSharedPtr<FControlBusDashboardEntry>>
+		, public TSharedFromThis<FControlBusTraceProvider>
+	{
+	public:
+		FControlBusTraceProvider(TSharedRef<UE::Audio::Insights::FTraceChannelManager> InManager)
+			: UE::Audio::Insights::TDeviceDataMapTraceProvider<uint32, TSharedPtr<FControlBusDashboardEntry>>(GetName_Static())
+		{
+			Channels.Add(InManager->CreateHandle({ TEXT("AudioChannel") }));
+		}
+
+		virtual ~FControlBusTraceProvider() = default;
+		virtual UE::Trace::IAnalyzer* ConstructAnalyzer() override;
+
+		virtual bool ProcessMessages() override;
+		static FName GetName_Static();
+
+	private:
+
+		FControlBusMessages TraceMessages;
+		TSet<UE::Audio::Insights::FTraceChannelHandle> Channels;
+	};
+} // namespace AudioModulationEditor
diff --git a/Engine/Plugins/Runtime/AudioModulation/Source/AudioModulationEditor/Private/Insights/Views/ControlBusDashboardViewFactory.cpp b/Engine/Plugins/Runtime/AudioModulation/Source/AudioModulationEditor/Private/Insights/Views/ControlBusDashboardViewFactory.cpp
new file mode 100644
--- /dev/null
+++ b/Engine/Plugins/Runtime/AudioModulation/Source/AudioModulationEditor/Private/Insights/Views/ControlBusDashboardViewFactory.cpp
@@ -0,0 +1,226 @@
+// Copyright Epic Games, Inc. All Rights Reserved.
+#include "ControlBusDashboardViewFactory.h"
+
+#include "AssetRegistry/AssetData.h"
+#include "AssetRegistry/AssetRegistryModule.h"
+#include "AudioDefines.h"
+#include "AudioDeviceManager.h"
+#include "AudioInsightsModule.h"
+#include "AudioInsightsStyle.h"
+#include "AudioInsightsTraceModule.h"
+#include "AudioInsightsTraceProviderBase.h"
+#include "Editor.h"
+#include "Internationalization/Text.h"
+#include "Insights/Providers/ControlBusTraceProvider.h"
+#include "Subsystems/AssetEditorSubsystem.h"
+#include "Templates/SharedPointer.h"
+
+#define LOCTEXT_NAMESPACE "AudioModulationInsights"
+
+namespace AudioModulationEditor
+{
+	namespace ControlBusPrivate
+	{
+		const FControlBusDashboardEntry& CastEntry(const UE::Audio::Insights::IDashboardDataViewEntry& InData)
+		{
+			return static_cast<const FControlBusDashboardEntry&>(InData);
+		};
+
+	} // namespace ControlBusPrivate
+
+	FControlBusDashboardViewFactory::FControlBusDashboardViewFactory()
+	{
+		IAudioInsightsModule& InsightsModule = FModuleManager::LoadModuleChecked<IAudioInsightsModule>(IAudioInsightsModule::GetEditorName());
+		IAudioInsightsTraceModule& InsightsTraceModule = InsightsModule.GetTraceModule();
+
+		TSharedPtr<FControlBusTraceProvider> ControlBusProvider = MakeShared<FControlBusTraceProvider>(InsightsTraceModule.GetChannelManager());
+		InsightsTraceModule.AddTraceProvider(StaticCastSharedPtr<UE::Audio::Insights::FTraceProviderBase>(ControlBusProvider));
+
+		Providers =
+		{
+			StaticCastSharedPtr<UE::Audio::Insights::FTraceProviderBase>(ControlBusProvider)
+		};
+	}
+
+	FControlBusDashboardViewFactory::~FControlBusDashboardViewFactory()
+	{
+	}
+
+	FName FControlBusDashboardViewFactory::GetName() const
+	{
+		return "ControlBuses";
+	}
+
+	FText FControlBusDashboardViewFactory::GetDisplayName() const
+	{
+		return LOCTEXT("AudioInsights_ModulationControlBus_DisplayName", "Control Buses");
+	}
+
+	void FControlBusDashboardViewFactory::ProcessEntries(UE::Audio::Insights::FTraceTableDashboardViewFactory::EProcessReason Reason)
+	{
+		using namespace UE::Audio::Insights;
+
+		const FString FilterString = GetSearchFilterText().ToString();
+		FTraceTableDashboardViewFactory::FilterEntries<FControlBusTraceProvider>([&FilterString](const IDashboardDataViewEntry& Entry)
+		{
+			const FControlBusDashboardEntry& ControlBusEntry = ControlBusPrivate::CastEntry(Entry);
+			if (ControlBusEntry.GetDisplayName().ToString().Contains(FilterString))
+			{
+				return false;
+			}
+
+			return true;
+		});
+	}
+
+	FSlateIcon FControlBusDashboardViewFactory::GetIcon() const
+	{
+		return { "AudioModulationStyle", "ClassIcon.SoundControlBus" };
+	}
+
+	UE::Audio::Insights::EDefaultDashboardTabStack FControlBusDashboardViewFactory::GetDefaultTabStack() const
+	{
+		return UE::Audio::Insights::EDefaultDashboardTabStack::Analysis;
+	}
+
+	TSharedRef<SWidget> FControlBusDashboardViewFactory::MakeWidget()
+	{
+		TSharedRef<SWidget> TableDashboardWidget = UE::Audio::Insights::FTraceTableDashboardViewFactory::MakeWidget();
+		TSharedRef<SWidget> BusWatchWidget = MakeControlBusWatchWidget();
+
+		return SNew(SHorizontalBox)
+			+ SHorizontalBox::Slot()
+			.HAlign(HAlign_Fill)
+			.Padding(0.0f, 6.0f, 0.0f, 0.0f)
+			[
+				// Dashboard and plots area
+				SNew(SSplitter)
+					.Orientation(Orient_Horizontal)
+					+ SSplitter::Slot()
+					.Value(0.55f)
+					[
+						TableDashboardWidget
+					]
+					+ SSplitter::Slot()
+					.Value(0.45f)
+					[
+						BusWatchWidget
+					]
+			];
+	}
+
+	const TMap<FName, UE::Audio::Insights::FTraceTableDashboardViewFactory::FColumnData>& FControlBusDashboardViewFactory::GetColumns() const
+	{
+		auto CreateColumnData = []()
+		{
+			using namespace UE::Audio::Insights;
+
+			return TMap<FName, FTraceTableDashboardViewFactory::FColumnData>
+			{
+				{
+					"BusId",
+					{
+						LOCTEXT("ControlBus_BusIdColumnDisplayName", "Control Bus ID"),
+						[](const IDashboardDataViewEntry& InData) { return FText::AsNumber(ControlBusPrivate::CastEntry(InData).ControlBusId); },
+						true /* bDefaultHidden */,
+						0.08f /* FillWidth */
+					}
+				},
+				{
+					"Name",
+					{
+						LOCTEXT("ControlBus_NameColumnDisplayName", "Name"),
+						[](const IDashboardDataViewEntry& InData) { return FText::FromString(ControlBusPrivate::CastEntry(InData).Name); },
+						false /* bDefaultHidden */,
+						0.75f /* FillWidth */
+					}
+				},
+				{
+					"ParameterName",
+					{
+						LOCTEXT("ControlBus_ParamNameColumnDisplayName", "Parameter"),
+						[](const IDashboardDataViewEntry& InData) { return FText::FromString(ControlBusPrivate::CastEntry(InData).ParamName); },
+						false /* bDefaultHidden */,
+						0.15f /* FillWidth */
+					}
+				},
+				{
+					"Value",
+					{
+						LOCTEXT("ControlBus_ValueColumnDisplayName", "Value"),
+						[](const IDashboardDataViewEntry& InData) { return FText::AsNumber(ControlBusPrivate::CastEntry(InData).Value, FSlateStyle::Get().GetAmpFloatFormat()); },
+						false /* bDefaultHidden */,
+						0.07f /* FillWidth */
+					}
+				}
+			};
+		};
+		static const TMap<FName, FTraceTableDashboardViewFactory::FColumnData> ColumnData = CreateColumnData();
+		return ColumnData;
+	}
+
+	void FControlBusDashboardViewFactory::SortTable()
+	{
+		using namespace UE::Audio::Insights;
+
+		if (SortByColumn == "Name")
+		{
+			if (SortMode == EColumnSortMode::Ascending)
+			{
+				DataViewEntries.Sort([](const TSharedPtr<IDashboardDataViewEntry>& A, const TSharedPtr<IDashboardDataViewEntry>& B)
+				{
+					const FControlBusDashboardEntry& AData = ControlBusPrivate::CastEntry(*A.Get());
+					const FControlBusDashboardEntry& BData = ControlBusPrivate::CastEntry(*B.Get());
+
+					return AData.GetDisplayName().CompareToCaseIgnored(BData.GetDisplayName()) < 0;
+				});
+			}
+			else if (SortMode == EColumnSortMode::Descending)
+			{
+				DataViewEntries.Sort([](const TSharedPtr<IDashboardDataViewEntry>& A, const TSharedPtr<IDashboardDataViewEntry>& B)
+				{
+					const FControlBusDashboardEntry& AData = ControlBusPrivate::CastEntry(*A.Get());
+					const FControlBusDashboardEntry& BData = ControlBusPrivate::CastEntry(*B.Get());
+
+					return BData.GetDisplayName().CompareToCaseIgnored(AData.GetDisplayName()) < 0;
+				});
+			}
+		}
+		else if (SortByColumn == "ParamName")
+		{
+			if (SortMode == EColumnSortMode::Ascending)
+			{
+				DataViewEntries.Sort([](const TSharedPtr<IDashboardDataViewEntry>& A, const TSharedPtr<IDashboardDataViewEntry>& B)
+				{
+					const FControlBusDashboardEntry& AData = ControlBusPrivate::CastEntry(*A.Get());
+					const FControlBusDashboardEntry& BData = ControlBusPrivate::CastEntry(*B.Get());
+
+					return AData.GetParamNameAsFText().CompareToCaseIgnored(BData.GetParamNameAsFText()) < 0;
+				});
+			}
+			else if (SortMode == EColumnSortMode::Descending)
+			{
+				DataViewEntries.Sort([](const TSharedPtr<IDashboardDataViewEntry>& A, const TSharedPtr<IDashboardDataViewEntry>& B)
+				{
+					const FControlBusDashboardEntry& AData = ControlBusPrivate::CastEntry(*A.Get());
+					const FControlBusDashboardEntry& BData = ControlBusPrivate::CastEntry(*B.Get());
+
+					return BData.GetParamNameAsFText().CompareToCaseIgnored(AData.GetParamNameAsFText()) < 0;
+				});
+			}
+		}
+	}
+
+	TSharedRef<SWidget> FControlBusDashboardViewFactory::MakeControlBusListWidget()
+	{
+		return SNew(SBox);
+	}
+
+	TSharedRef<SWidget> FControlBusDashboardViewFactory::MakeControlBusWatchWidget()
+	{
+		return SNew(SBox);
+	}
+
+} // namespace AudioModulationEditor
+
+#undef LOCTEXT_NAMESPACE
diff --git a/Engine/Plugins/Runtime/AudioModulation/Source/AudioModulationEditor/Private/Insights/Views/ControlBusDashboardViewFactory.h b/Engine/Plugins/Runtime/AudioModulation/Source/AudioModulationEditor/Private/Insights/Views/ControlBusDashboardViewFactory.h
new file mode 100644
--- /dev/null
+++ b/Engine/Plugins/Runtime/AudioModulation/Source/AudioModulationEditor/Private/Insights/Views/ControlBusDashboardViewFactory.h
@@ -0,0 +1,33 @@
+// Copyright Epic Games, Inc. All Rights Reserved.
+#pragma once
+
+#include "Delegates/Delegate.h"
+#include "Templates/SharedPointer.h"
+#include "Views/TableDashboardViewFactory.h"
+#include "Widgets/Input/SCheckBox.h"
+
+namespace AudioModulationEditor
+{
+	class FControlBusDashboardViewFactory : public UE::Audio::Insights::FTraceObjectTableDashboardViewFactory
+	{
+	public:
+		FControlBusDashboardViewFactory();
+		virtual ~FControlBusDashboardViewFactory();
+
+		virtual FName GetName() const override;
+		virtual FText GetDisplayName() const override;
+		virtual FSlateIcon GetIcon() const override;
+		virtual UE::Audio::Insights::EDefaultDashboardTabStack GetDefaultTabStack() const override;
+		virtual void ProcessEntries(UE::Audio::Insights::FTraceTableDashboardViewFactory::EProcessReason Reason) override;
+		virtual TSharedRef<SWidget> MakeWidget() override;
+
+	protected:
+		virtual const TMap<FName, UE::Audio::Insights::FTraceTableDashboardViewFactory::FColumnData>& GetColumns() const override;
+
+		virtual void SortTable() override;
+
+		virtual TSharedRef<SWidget> MakeControlBusListWidget();
+		virtual TSharedRef<SWidget> MakeControlBusWatchWidget();
+
+	};
+} // namespace AudioModulationEditor
diff --git a/Engine/Plugins/Runtime/AudioModulation/Source/AudioModulationEditor/Private/Insights/Views/ModulationDashboardViewFactory.cpp b/Engine/Plugins/Runtime/AudioModulation/Source/AudioModulationEditor/Private/Insights/Views/ModulationDashboardViewFactory.cpp
deleted file mode 100644
diff --git a/Engine/Plugins/Runtime/AudioModulation/Source/AudioModulationEditor/Private/Insights/Views/ModulationDashboardViewFactory.h b/Engine/Plugins/Runtime/AudioModulation/Source/AudioModulationEditor/Private/Insights/Views/ModulationDashboardViewFactory.h
deleted file mode 100644
