[Audio Insights] Update Audio Insights data when scrubbing in Unreal Insights (part 1): Passing Session To Analyzers to be able to update the session duration.

diff --git a/Engine/Plugins/AudioInsights/Source/AudioInsights/Private/AudioInsightsTraceModule.cpp b/Engine/Plugins/AudioInsights/Source/AudioInsights/Private/AudioInsightsTraceModule.cpp
--- a/Engine/Plugins/AudioInsights/Source/AudioInsights/Private/AudioInsightsTraceModule.cpp
+++ b/Engine/Plugins/AudioInsights/Source/AudioInsights/Private/AudioInsightsTraceModule.cpp
@@ -21,8 +21,8 @@ namespace UE::Audio::Insights
 			TSharedPtr<FMixerSourceTraceProvider> SourceProvider = MakeShared<FMixerSourceTraceProvider>(ChannelManager);
 			TSharedPtr<FVirtualLoopTraceProvider> VirtualLoopProvider = MakeShared<FVirtualLoopTraceProvider>(ChannelManager);
 
-			TraceProviders.Add(SourceProvider->GetName(), StaticCastSharedPtr<FTraceProviderBase>(SourceProvider));
-			TraceProviders.Add(VirtualLoopProvider->GetName(), StaticCastSharedPtr<FTraceProviderBase>(VirtualLoopProvider));
+			TraceProviders.Add(SourceProvider->GetName(), SourceProvider);
+			TraceProviders.Add(VirtualLoopProvider->GetName(), VirtualLoopProvider);
 		}
 	}
 
@@ -69,10 +69,10 @@ namespace UE::Audio::Insights
 
 	void FTraceModule::OnAnalysisBegin(TraceServices::IAnalysisSession& InSession)
 	{
-		for (const TPair<FName, TSharedPtr<FTraceProviderBase>>& Pair : TraceProviders)
+		for (const auto& [ProviderName, Provider] : TraceProviders)
 		{
-			InSession.AddProvider(Pair.Key, nullptr, Pair.Value);
-			InSession.AddAnalyzer(Pair.Value->ConstructAnalyzer());
+			InSession.AddProvider(ProviderName, Provider, Provider);
+			InSession.AddAnalyzer(Provider->ConstructAnalyzer(InSession));
 		}
 
 		FirstTimeStamp = -TNumericLimits<double>::Min();
diff --git a/Engine/Plugins/AudioInsights/Source/AudioInsights/Private/Providers/MixerSourceTraceProvider.cpp b/Engine/Plugins/AudioInsights/Source/AudioInsights/Private/Providers/MixerSourceTraceProvider.cpp
--- a/Engine/Plugins/AudioInsights/Source/AudioInsights/Private/Providers/MixerSourceTraceProvider.cpp
+++ b/Engine/Plugins/AudioInsights/Source/AudioInsights/Private/Providers/MixerSourceTraceProvider.cpp
@@ -139,13 +139,14 @@ namespace UE::Audio::Insights
 		return true;
 	}
 
-	UE::Trace::IAnalyzer* FMixerSourceTraceProvider::ConstructAnalyzer()
+	UE::Trace::IAnalyzer* FMixerSourceTraceProvider::ConstructAnalyzer(TraceServices::IAnalysisSession& InSession)
 	{
 		class FMixerSourceTraceAnalyzer : public FTraceAnalyzerBase
 		{
 		public:
-			FMixerSourceTraceAnalyzer(TSharedRef<FMixerSourceTraceProvider> InProvider)
+			FMixerSourceTraceAnalyzer(TSharedRef<FMixerSourceTraceProvider> InProvider, TraceServices::IAnalysisSession& InSession)
 				: FTraceAnalyzerBase(InProvider)
+				, Session(InSession)
 			{
 			}
 
@@ -232,6 +233,13 @@ namespace UE::Audio::Insights
 					}
 				}
 
+				const double Timestamp = Context.EventTime.AsSeconds(Context.EventData.GetValue<uint64>("Timestamp"));
+
+				{
+					TraceServices::FAnalysisSessionEditScope SessionEditScope(Session);
+					Session.UpdateDurationSeconds(Timestamp);
+				}
+
 				return OnEventSuccess(RouteId, Style, Context);
 			}
 
@@ -248,8 +256,10 @@ namespace UE::Audio::Insights
 			};
 
 			bool bIsFirstTimestampSet = false;
+
+			TraceServices::IAnalysisSession& Session;
 		};
 
-		return new FMixerSourceTraceAnalyzer(AsShared());
+		return new FMixerSourceTraceAnalyzer(AsShared(), InSession);
 	}
 } // namespace UE::Audio::Insights
diff --git a/Engine/Plugins/AudioInsights/Source/AudioInsights/Private/Providers/MixerSourceTraceProvider.h b/Engine/Plugins/AudioInsights/Source/AudioInsights/Private/Providers/MixerSourceTraceProvider.h
--- a/Engine/Plugins/AudioInsights/Source/AudioInsights/Private/Providers/MixerSourceTraceProvider.h
+++ b/Engine/Plugins/AudioInsights/Source/AudioInsights/Private/Providers/MixerSourceTraceProvider.h
@@ -21,7 +21,7 @@ namespace UE::Audio::Insights
 		}
 
 		virtual ~FMixerSourceTraceProvider() = default;
-		virtual UE::Trace::IAnalyzer* ConstructAnalyzer() override;
+		virtual UE::Trace::IAnalyzer* ConstructAnalyzer(TraceServices::IAnalysisSession& InSession) override;
 
 		virtual bool ProcessMessages() override;
 
diff --git a/Engine/Plugins/AudioInsights/Source/AudioInsights/Private/Providers/VirtualLoopTraceProvider.cpp b/Engine/Plugins/AudioInsights/Source/AudioInsights/Private/Providers/VirtualLoopTraceProvider.cpp
--- a/Engine/Plugins/AudioInsights/Source/AudioInsights/Private/Providers/VirtualLoopTraceProvider.cpp
+++ b/Engine/Plugins/AudioInsights/Source/AudioInsights/Private/Providers/VirtualLoopTraceProvider.cpp
@@ -78,13 +78,14 @@ namespace UE::Audio::Insights
 		return true;
 	}
 
-	UE::Trace::IAnalyzer* FVirtualLoopTraceProvider::ConstructAnalyzer()
+	UE::Trace::IAnalyzer* FVirtualLoopTraceProvider::ConstructAnalyzer(TraceServices::IAnalysisSession& InSession)
 	{
 		class FVirtualLoopTraceAnalyzer : public FTraceAnalyzerBase
 		{
 		public:
-			FVirtualLoopTraceAnalyzer(TSharedRef<FVirtualLoopTraceProvider> InProvider)
+			FVirtualLoopTraceAnalyzer(TSharedRef<FVirtualLoopTraceProvider> InProvider, TraceServices::IAnalysisSession& InSession)
 				: FTraceAnalyzerBase(InProvider)
+				, Session(InSession)
 			{
 			}
 
@@ -136,6 +137,13 @@ namespace UE::Audio::Insights
 					}
 				}
 
+				const double Timestamp = Context.EventTime.AsSeconds(Context.EventData.GetValue<uint64>("Timestamp"));
+
+				{
+					TraceServices::FAnalysisSessionEditScope SessionEditScope(Session);
+					Session.UpdateDurationSeconds(Timestamp);
+				}
+
 				return OnEventSuccess(RouteId, Style, Context);
 			}
 
@@ -147,8 +155,10 @@ namespace UE::Audio::Insights
 				RouteId_Realize,
 				RouteId_Stop
 			};
+
+			TraceServices::IAnalysisSession& Session;
 		};
 
-		return new FVirtualLoopTraceAnalyzer(AsShared());
+		return new FVirtualLoopTraceAnalyzer(AsShared(), InSession);
 	}
 } // namespace UE::Audio::Insights
diff --git a/Engine/Plugins/AudioInsights/Source/AudioInsights/Private/Providers/VirtualLoopTraceProvider.h b/Engine/Plugins/AudioInsights/Source/AudioInsights/Private/Providers/VirtualLoopTraceProvider.h
--- a/Engine/Plugins/AudioInsights/Source/AudioInsights/Private/Providers/VirtualLoopTraceProvider.h
+++ b/Engine/Plugins/AudioInsights/Source/AudioInsights/Private/Providers/VirtualLoopTraceProvider.h
@@ -20,7 +20,7 @@ namespace UE::Audio::Insights
 		}
 
 		virtual ~FVirtualLoopTraceProvider() = default;
-		virtual UE::Trace::IAnalyzer* ConstructAnalyzer() override;
+		virtual UE::Trace::IAnalyzer* ConstructAnalyzer(TraceServices::IAnalysisSession& InSession) override;
 
 		static FName GetName_Static();
 
diff --git a/Engine/Plugins/AudioInsights/Source/AudioInsights/Public/AudioInsightsTraceProviderBase.h b/Engine/Plugins/AudioInsights/Source/AudioInsights/Public/AudioInsightsTraceProviderBase.h
--- a/Engine/Plugins/AudioInsights/Source/AudioInsights/Public/AudioInsightsTraceProviderBase.h
+++ b/Engine/Plugins/AudioInsights/Source/AudioInsights/Public/AudioInsightsTraceProviderBase.h
@@ -16,14 +16,13 @@
 
 namespace UE::Audio::Insights
 {
-	class AUDIOINSIGHTS_API FTraceProviderBase
-		: public TraceServices::IEditableProvider
+	class AUDIOINSIGHTS_API FTraceProviderBase : public TraceServices::IProvider, public TraceServices::IEditableProvider
 	{
 	public:
-		FTraceProviderBase(FName InName);
+		explicit FTraceProviderBase(FName InName);
 		virtual ~FTraceProviderBase() = default;
 
-		virtual Trace::IAnalyzer* ConstructAnalyzer() = 0;
+		virtual Trace::IAnalyzer* ConstructAnalyzer(TraceServices::IAnalysisSession& InSession) = 0;
 		FName GetName() const;
 
 		virtual void Reset()
diff --git a/Engine/Plugins/AudioInsights/Source/AudioInsightsEditor/Private/Providers/AudioBusProvider.cpp b/Engine/Plugins/AudioInsights/Source/AudioInsightsEditor/Private/Providers/AudioBusProvider.cpp
--- a/Engine/Plugins/AudioInsights/Source/AudioInsightsEditor/Private/Providers/AudioBusProvider.cpp
+++ b/Engine/Plugins/AudioInsights/Source/AudioInsightsEditor/Private/Providers/AudioBusProvider.cpp
@@ -160,7 +160,7 @@ namespace UE::Audio::Insights
 		return true;
 	}
 
-	UE::Trace::IAnalyzer* FAudioBusProvider::ConstructAnalyzer()
+	UE::Trace::IAnalyzer* FAudioBusProvider::ConstructAnalyzer(TraceServices::IAnalysisSession& InSession)
 	{
 		return nullptr;
 	}
diff --git a/Engine/Plugins/AudioInsights/Source/AudioInsightsEditor/Private/Providers/AudioBusProvider.h b/Engine/Plugins/AudioInsights/Source/AudioInsightsEditor/Private/Providers/AudioBusProvider.h
--- a/Engine/Plugins/AudioInsights/Source/AudioInsightsEditor/Private/Providers/AudioBusProvider.h
+++ b/Engine/Plugins/AudioInsights/Source/AudioInsightsEditor/Private/Providers/AudioBusProvider.h
@@ -15,7 +15,7 @@ namespace UE::Audio::Insights
 
 		static FName GetName_Static();
 
-		virtual UE::Trace::IAnalyzer* ConstructAnalyzer() override;
+		virtual UE::Trace::IAnalyzer* ConstructAnalyzer(TraceServices::IAnalysisSession& InSession) override;
 		
 		DECLARE_MULTICAST_DELEGATE_OneParam(FOnAudioBusAssetAdded, const TWeakObjectPtr<UObject> /*Asset*/);
 		inline static FOnAudioBusAssetAdded OnAudioBusAssetAdded;
diff --git a/Engine/Plugins/AudioInsights/Source/AudioInsightsEditor/Private/Providers/SoundSubmixProvider.cpp b/Engine/Plugins/AudioInsights/Source/AudioInsightsEditor/Private/Providers/SoundSubmixProvider.cpp
--- a/Engine/Plugins/AudioInsights/Source/AudioInsightsEditor/Private/Providers/SoundSubmixProvider.cpp
+++ b/Engine/Plugins/AudioInsights/Source/AudioInsightsEditor/Private/Providers/SoundSubmixProvider.cpp
@@ -160,7 +160,7 @@ namespace UE::Audio::Insights
 		return true;
 	}
 
-	UE::Trace::IAnalyzer* FSoundSubmixProvider::ConstructAnalyzer()
+	UE::Trace::IAnalyzer* FSoundSubmixProvider::ConstructAnalyzer(TraceServices::IAnalysisSession& InSession)
 	{
 		return nullptr;
 	}
diff --git a/Engine/Plugins/AudioInsights/Source/AudioInsightsEditor/Private/Providers/SoundSubmixProvider.h b/Engine/Plugins/AudioInsights/Source/AudioInsightsEditor/Private/Providers/SoundSubmixProvider.h
--- a/Engine/Plugins/AudioInsights/Source/AudioInsightsEditor/Private/Providers/SoundSubmixProvider.h
+++ b/Engine/Plugins/AudioInsights/Source/AudioInsightsEditor/Private/Providers/SoundSubmixProvider.h
@@ -15,7 +15,7 @@ namespace UE::Audio::Insights
 
 		static FName GetName_Static();
 
-		virtual UE::Trace::IAnalyzer* ConstructAnalyzer() override;
+		virtual UE::Trace::IAnalyzer* ConstructAnalyzer(TraceServices::IAnalysisSession& InSession) override;
 		
 		DECLARE_MULTICAST_DELEGATE_OneParam(FOnSubmixAssetAdded, const TWeakObjectPtr<UObject> /*Asset*/);
 		inline static FOnSubmixAssetAdded OnSubmixAssetAdded;
diff --git a/Engine/Plugins/Runtime/AudioModulation/Source/AudioModulationEditor/Private/Insights/Providers/ControlBusTraceProvider.cpp b/Engine/Plugins/Runtime/AudioModulation/Source/AudioModulationEditor/Private/Insights/Providers/ControlBusTraceProvider.cpp
--- a/Engine/Plugins/Runtime/AudioModulation/Source/AudioModulationEditor/Private/Insights/Providers/ControlBusTraceProvider.cpp
+++ b/Engine/Plugins/Runtime/AudioModulation/Source/AudioModulationEditor/Private/Insights/Providers/ControlBusTraceProvider.cpp
@@ -65,13 +65,14 @@ namespace AudioModulationEditor
 		return true;
 	}
 
-	UE::Trace::IAnalyzer* FControlBusTraceProvider::ConstructAnalyzer()
+	UE::Trace::IAnalyzer* FControlBusTraceProvider::ConstructAnalyzer(TraceServices::IAnalysisSession& InSession)
 	{
 		class FControlBusTraceAnalyzer : public UE::Audio::Insights::FTraceProviderBase::FTraceAnalyzerBase
 		{
 		public:
-			FControlBusTraceAnalyzer(TSharedRef<FControlBusTraceProvider> InProvider)
+			FControlBusTraceAnalyzer(TSharedRef<FControlBusTraceProvider> InProvider, TraceServices::IAnalysisSession& InSession)
 				: UE::Audio::Insights::FTraceProviderBase::FTraceAnalyzerBase(InProvider)
+				, Session(InSession)
 			{
 			}
 
@@ -116,6 +117,13 @@ namespace AudioModulationEditor
 					}
 				}
 
+				const double Timestamp = Context.EventTime.AsSeconds(Context.EventData.GetValue<uint64>("Timestamp"));
+
+				{
+					TraceServices::FAnalysisSessionEditScope SessionEditScope(Session);
+					Session.UpdateDurationSeconds(Timestamp);
+				}
+
 				return OnEventSuccess(RouteId, Style, Context);
 			}
 
@@ -126,8 +134,10 @@ namespace AudioModulationEditor
 				RouteId_Deactivate,
 				RouteId_Update
 			};
+
+			TraceServices::IAnalysisSession& Session;
 		};
 
-		return new FControlBusTraceAnalyzer(AsShared());
+		return new FControlBusTraceAnalyzer(AsShared(), InSession);
 	}
 } // namespace AudioModulationEditor
diff --git a/Engine/Plugins/Runtime/AudioModulation/Source/AudioModulationEditor/Private/Insights/Providers/ControlBusTraceProvider.h b/Engine/Plugins/Runtime/AudioModulation/Source/AudioModulationEditor/Private/Insights/Providers/ControlBusTraceProvider.h
--- a/Engine/Plugins/Runtime/AudioModulation/Source/AudioModulationEditor/Private/Insights/Providers/ControlBusTraceProvider.h
+++ b/Engine/Plugins/Runtime/AudioModulation/Source/AudioModulationEditor/Private/Insights/Providers/ControlBusTraceProvider.h
@@ -20,7 +20,7 @@ namespace AudioModulationEditor
 		}
 
 		virtual ~FControlBusTraceProvider() = default;
-		virtual UE::Trace::IAnalyzer* ConstructAnalyzer() override;
+		virtual UE::Trace::IAnalyzer* ConstructAnalyzer(TraceServices::IAnalysisSession& InSession) override;
 
 		virtual bool ProcessMessages() override;
 		static FName GetName_Static();
