New Control Buses default to Volume Parameter

diff --git a/Engine/Plugins/Runtime/AudioModulation/Source/AudioModulation/Private/AudioModulationSettings.cpp b/Engine/Plugins/Runtime/AudioModulation/Source/AudioModulation/Private/AudioModulationSettings.cpp
--- a/Engine/Plugins/Runtime/AudioModulation/Source/AudioModulation/Private/AudioModulationSettings.cpp
+++ b/Engine/Plugins/Runtime/AudioModulation/Source/AudioModulation/Private/AudioModulationSettings.cpp
@@ -74,3 +74,19 @@ void UAudioModulationSettings::RegisterParameters() const
 		}
 	}
 }
+
+TObjectPtr<USoundModulationParameter> UAudioModulationSettings::GetModulationParameter(const FString& InName) const
+{
+	for (const FSoftObjectPath& Path : Parameters)
+	{
+		if (USoundModulationParameter* Param = Cast<USoundModulationParameter>(Path.TryLoad()))
+		{
+			if (Param->GetName() == InName)
+			{
+				return Param;
+			}
+		}
+	}
+
+	return nullptr;
+}
diff --git a/Engine/Plugins/Runtime/AudioModulation/Source/AudioModulation/Public/AudioModulationSettings.h b/Engine/Plugins/Runtime/AudioModulation/Source/AudioModulation/Public/AudioModulationSettings.h
--- a/Engine/Plugins/Runtime/AudioModulation/Source/AudioModulation/Public/AudioModulationSettings.h
+++ b/Engine/Plugins/Runtime/AudioModulation/Source/AudioModulation/Public/AudioModulationSettings.h
@@ -3,6 +3,7 @@
 
 #include "Engine/DeveloperSettings.h"
 #include "Misc/Build.h"
+#include "SoundModulationParameter.h"
 #include "UObject/ObjectMacros.h"
 #include "UObject/SoftObjectPath.h"
 #include "UObject/UnrealType.h"
@@ -22,6 +23,8 @@ public:
 
 	void RegisterParameters() const;
 
+	TObjectPtr<USoundModulationParameter> GetModulationParameter(const FString& InName) const;
+
 #if WITH_EDITOR
 	virtual void PostEditChangeChainProperty(FPropertyChangedChainEvent& PropertyChangedEvent) override;
 #endif // WITH_EDITOR
diff --git a/Engine/Plugins/Runtime/AudioModulation/Source/AudioModulationEditor/Private/Factories/SoundControlBusFactory.cpp b/Engine/Plugins/Runtime/AudioModulation/Source/AudioModulationEditor/Private/Factories/SoundControlBusFactory.cpp
--- a/Engine/Plugins/Runtime/AudioModulation/Source/AudioModulationEditor/Private/Factories/SoundControlBusFactory.cpp
+++ b/Engine/Plugins/Runtime/AudioModulation/Source/AudioModulationEditor/Private/Factories/SoundControlBusFactory.cpp
@@ -3,6 +3,7 @@
 #include "SoundControlBusFactory.h"
 
 #include "AudioAnalytics.h"
+#include "AudioModulationSettings.h"
 #include "SoundControlBus.h"
 
 #include UE_INLINE_GENERATED_CPP_BY_NAME(SoundControlBusFactory)
@@ -20,5 +21,13 @@ USoundControlBusFactory::USoundControlBusFactory(const FObjectInitializer& Objec
 UObject* USoundControlBusFactory::FactoryCreateNew(UClass* Class, UObject* InParent, FName Name, EObjectFlags Flags, UObject* Context, FFeedbackContext* Warn)
 {
 	Audio::Analytics::RecordEvent_Usage(TEXT("AudioModulation.ControlBusCreated"));
-	return NewObject<USoundControlBus>(InParent, Name, Flags);
+	USoundControlBus* NewControlBus = NewObject<USoundControlBus>(InParent, Name, Flags);
+	if (NewControlBus)
+	{
+		if (const UAudioModulationSettings* ModulationSettings = GetDefault<UAudioModulationSettings>())
+		{
+			NewControlBus->Parameter = ModulationSettings->GetModulationParameter("Volume");
+		}
+	}
+	return NewControlBus;
 }
