CVar to Disable Audio Modulation

diff --git a/Engine/Plugins/Runtime/AudioModulation/Source/AudioModulation/Private/AudioModulation.cpp b/Engine/Plugins/Runtime/AudioModulation/Source/AudioModulation/Private/AudioModulation.cpp
--- a/Engine/Plugins/Runtime/AudioModulation/Source/AudioModulation/Private/AudioModulation.cpp
+++ b/Engine/Plugins/Runtime/AudioModulation/Source/AudioModulation/Private/AudioModulation.cpp
@@ -24,6 +24,16 @@ REGISTER_METASOUND_DATATYPE(AudioModulation::FSoundModulatorAsset, "Modulator",
 REGISTER_METASOUND_DATATYPE(AudioModulation::FSoundModulationParameterAsset, "ModulationParameter", Metasound::ELiteralType::UObjectProxy, USoundModulationParameter);
 #endif // WITH_AUDIOMODULATION_METASOUND_SUPPORT
 
+namespace AudioModulationCVars
+{
+	static bool bAudioModulationEnabledCVar = true;
+	FAutoConsoleVariableRef CVarAudioModulationEnabled(
+		TEXT("au.EnableAudioModulation"),
+		bAudioModulationEnabledCVar,
+		TEXT("Set to 0 to disable Audio Modulation entirely.\n"),
+		ECVF_Default);
+}
+
 namespace AudioModulation
 {
 	FAudioModulationManager::FAudioModulationManager()
@@ -48,7 +58,7 @@ namespace AudioModulation
 
 	void FAudioModulationManager::ActivateBus(const USoundControlBus& InBus)
 	{
-		UE_LOG(LogAudioModulation, Warning, TEXT("FAudioModulationManager::ActivateBus is deprecated. Use UAudioModulationDestination API (see SetModulator) instead to enforce generator lifetime"));
+		UE_LOG(LogAudioModulation, Warning, TEXT("FAudioModulationManager::ActivateBus is deprecated. Use UAudioModulationDestination API (see SetModulator) instead to enforce bus lifetime"));
 
 		PRAGMA_DISABLE_DEPRECATION_WARNINGS
 		ModSystem->ActivateBus(InBus);
@@ -95,7 +105,7 @@ namespace AudioModulation
 
 	void FAudioModulationManager::DeactivateGenerator(const USoundModulationGenerator& InGenerator)
 	{
-		UE_LOG(LogAudioModulation, Warning, TEXT("FAudioModulationManager::DeactivateGenerator is deprecated. Use UAudioModulationDestination API (see ClearModulator) instead to enforce bus lifetime"));
+		UE_LOG(LogAudioModulation, Warning, TEXT("FAudioModulationManager::DeactivateGenerator is deprecated. Use UAudioModulationDestination API (see ClearModulator) instead to enforce generator lifetime"));
 
 		PRAGMA_DISABLE_DEPRECATION_WARNINGS
 		ModSystem->DeactivateGenerator(InGenerator);
@@ -211,6 +221,10 @@ namespace AudioModulation
 
 	void FAudioModulationManager::ProcessModulators(const double InElapsed)
 	{
+		if (!AudioModulationCVars::bAudioModulationEnabledCVar)
+		{
+			return;
+		}
 		ModSystem->ProcessModulators(InElapsed);
 	}
 
diff --git a/Engine/Plugins/Runtime/AudioModulation/Source/AudioModulation/Public/AudioModulation.h b/Engine/Plugins/Runtime/AudioModulation/Source/AudioModulation/Public/AudioModulation.h
--- a/Engine/Plugins/Runtime/AudioModulation/Source/AudioModulation/Public/AudioModulation.h
+++ b/Engine/Plugins/Runtime/AudioModulation/Source/AudioModulation/Public/AudioModulation.h
@@ -59,23 +59,23 @@ namespace AudioModulation
 	//~ End IAudioModulationManager implementation
 
 	public:
-		UE_DEPRECATED(5.4, "Deactivation of modulators in this manner is now deprecated. Use USoundModulationWatchers to safety activate and track a given modulator")
+		UE_DEPRECATED(5.4, "Deactivation of modulators in this manner is now deprecated. Use USoundModulationDestination to safety activate and track a given modulator")
 		void ActivateBus(const USoundControlBus& InBus);
 
 		void ActivateBusMix(const USoundControlBusMix& InBusMix);
 
-		UE_DEPRECATED(5.4, "Deactivation of modulators in this manner is now deprecated. Use USoundModulationWatchers to safety activate and track a given modulator")
+		UE_DEPRECATED(5.4, "Deactivation of modulators in this manner is now deprecated. Use USoundModulationDestination to safety activate and track a given modulator")
 		void ActivateGenerator(const USoundModulationGenerator& InGenerator);
 
 		USoundControlBusMix* CreateBusMixFromValue(FName Name, const TArray<USoundControlBus*>& Buses, float Value, float AttackTime = -1.0f, float ReleaseTime = -1.0f);
 
-		UE_DEPRECATED(5.4, "Deactivation of modulators in this manner is now deprecated. Use USoundModulationWatchers to safety activate and track a given modulator")
+		UE_DEPRECATED(5.4, "Deactivation of modulators in this manner is now deprecated. Use USoundModulationDestination to safety activate and track a given modulator")
 		void DeactivateBus(const USoundControlBus& InBus);
 
 		void DeactivateBusMix(const USoundControlBusMix& InBusMix);
 		void DeactivateAllBusMixes();
 
-		UE_DEPRECATED(5.4, "Deactivation of modulators in this manner is now deprecated. Use USoundModulationWatchers to safety activate and track a given modulator")
+		UE_DEPRECATED(5.4, "Deactivation of modulators in this manner is now deprecated. Use USoundModulationDestination to safety activate and track a given modulator")
 		void DeactivateGenerator(const USoundModulationGenerator& InGenerator);
 
 		bool IsBusMixActive(const USoundControlBusMix& InBusMix);
