Update tooltips for deprecated BP Modulation functions related to Modulation Destionation functions

diff --git a/Engine/Plugins/Runtime/AudioModulation/Source/AudioModulation/Public/AudioModulationStatics.h b/Engine/Plugins/Runtime/AudioModulation/Source/AudioModulation/Public/AudioModulationStatics.h
--- a/Engine/Plugins/Runtime/AudioModulation/Source/AudioModulation/Public/AudioModulationStatics.h
+++ b/Engine/Plugins/Runtime/AudioModulation/Source/AudioModulation/Public/AudioModulationStatics.h
@@ -39,7 +39,7 @@ public:
 	 */
 	static AudioModulation::FAudioModulationManager* GetModulation(UWorld* World);
 
-	/** SOFT DEPRECATED: Use CreateModulationWatcher and store resulting watcher instead!
+	/** SOFT DEPRECATED: Use CreateModulationDestination and store resulting Destination instead!
 	 * Manually activates a modulation bus. If called, deactivation will only occur
 	 * if bus is manually deactivated or destroyed (i.e. will not deactivate
 	 * when all references become inactive).
@@ -61,7 +61,7 @@ public:
 	)
 	static void ActivateBusMix(const UObject* WorldContextObject, USoundControlBusMix* Mix);
 
-	/** SOFT DEPRECATED: Use CreateModulationWatcher and store resulting watcher instead!
+	/** SOFT DEPRECATED: Use CreateModulationDestination and store resulting Destination instead!
 	 * Manually activates a modulation generator. If called, deactivation will only occur
 	 * if generator is manually deactivated and not referenced or destroyed (i.e. will not deactivate
 	 * when all references become inactive).
@@ -216,7 +216,7 @@ public:
 	 * Deactivates a bus. Does nothing if the provided bus is already inactive.
 	 * @param Bus - Scope of modulator
 	 */
-	UFUNCTION(BlueprintCallable, Category = "Audio|Modulation", DisplayName = "Deactivate Control Bus", meta = (
+	UFUNCTION(BlueprintCallable, Category = "Audio|Modulation", DisplayName = "Deactivate Control Bus (Deprecated - 5.4)", meta = (
 		WorldContext = "WorldContextObject", 
 		Keywords = "modulation modulator bus")
 	)
