Add Bus Name & Parameter Name to Control Bus Update Message in Audio Insights

diff --git a/Engine/Plugins/Runtime/AudioModulation/Source/AudioModulation/Private/AudioModulationSystem.cpp b/Engine/Plugins/Runtime/AudioModulation/Source/AudioModulation/Private/AudioModulationSystem.cpp
--- a/Engine/Plugins/Runtime/AudioModulation/Source/AudioModulation/Private/AudioModulationSystem.cpp
+++ b/Engine/Plugins/Runtime/AudioModulation/Source/AudioModulation/Private/AudioModulationSystem.cpp
@@ -41,6 +41,8 @@ UE_TRACE_EVENT_BEGIN(Audio, ControlBusUpdate)
 	UE_TRACE_EVENT_FIELD(uint32, DeviceId)
 	UE_TRACE_EVENT_FIELD(uint32, ControlBusId)
 	UE_TRACE_EVENT_FIELD(double, Timestamp)
+	UE_TRACE_EVENT_FIELD(UE::Trace::WideString, Name)
+	UE_TRACE_EVENT_FIELD(UE::Trace::WideString, ParamName)
 	UE_TRACE_EVENT_FIELD(float, Value)
 UE_TRACE_EVENT_END()
 
@@ -573,6 +575,8 @@ namespace AudioModulation
 					<< ControlBusUpdate.DeviceId(static_cast<uint32>(AudioDeviceId))
 					<< ControlBusUpdate.ControlBusId(static_cast<uint32>(Pair.Key))
 					<< ControlBusUpdate.Timestamp(FPlatformTime::Cycles64())
+					<< ControlBusUpdate.Name(*(Pair.Value.GetName().ToString()))
+					<< ControlBusUpdate.ParamName(*(Pair.Value.GetParameterName().ToString()))
 					<< ControlBusUpdate.Value(Pair.Value.GetValue());
 			}
 		}
diff --git a/Engine/Plugins/Runtime/AudioModulation/Source/AudioModulation/Private/SoundControlBusProxy.cpp b/Engine/Plugins/Runtime/AudioModulation/Source/AudioModulation/Private/SoundControlBusProxy.cpp
--- a/Engine/Plugins/Runtime/AudioModulation/Source/AudioModulation/Private/SoundControlBusProxy.cpp
+++ b/Engine/Plugins/Runtime/AudioModulation/Source/AudioModulation/Private/SoundControlBusProxy.cpp
@@ -86,6 +86,15 @@ namespace AudioModulation
 		return FMath::Clamp(DefaultMixed * GeneratorValue, 0.0f, 1.0f);
 	}
 
+	FName FControlBusProxy::GetParameterName() const
+	{
+#if UE_BUILD_SHIPPING
+		static FName ParameterName;
+#endif // !UE_BUILD_SHIPPING
+
+		return ParameterName;
+	}
+
 	void FControlBusProxy::Init(FControlBusSettings&& InSettings)
 	{
 		check(ModSystem);
@@ -94,6 +103,10 @@ namespace AudioModulation
 		MixValue = NAN;
 		MixFunction = MoveTemp(InSettings.MixFunction);
 
+#if !UE_BUILD_SHIPPING
+		ParameterName = InSettings.OutputParameter.ParameterName;
+#endif // !UE_BUILD_SHIPPING 
+
 		DefaultValue = FMath::Clamp(InSettings.DefaultValue, 0.0f, 1.0f);
 		bBypass = InSettings.bBypass;
 
diff --git a/Engine/Plugins/Runtime/AudioModulation/Source/AudioModulation/Private/SoundControlBusProxy.h b/Engine/Plugins/Runtime/AudioModulation/Source/AudioModulation/Private/SoundControlBusProxy.h
--- a/Engine/Plugins/Runtime/AudioModulation/Source/AudioModulation/Private/SoundControlBusProxy.h
+++ b/Engine/Plugins/Runtime/AudioModulation/Source/AudioModulation/Private/SoundControlBusProxy.h
@@ -74,6 +74,7 @@ namespace AudioModulation
 		float GetGeneratorValue() const;
 		float GetMixValue() const;
 		float GetValue() const;
+		FName GetParameterName() const;
 		bool IsBypassed() const;
 		void MixIn(const float InValue);
 		void MixGenerators();
@@ -93,6 +94,11 @@ namespace AudioModulation
 
 		Audio::FModulationMixFunction MixFunction;
 		TArray<FGeneratorHandle> GeneratorHandles;
+
+#if !UE_BUILD_SHIPPING
+		FName ParameterName;
+#endif // !UE_BUILD_SHIPPING 
+
 	};
 
 	using FBusProxyMap = TMap<FBusId, FControlBusProxy>;
diff --git a/Engine/Plugins/Runtime/AudioModulation/Source/AudioModulationEditor/Private/Insights/Messages/ControlBusTraceMessages.h b/Engine/Plugins/Runtime/AudioModulation/Source/AudioModulationEditor/Private/Insights/Messages/ControlBusTraceMessages.h
--- a/Engine/Plugins/Runtime/AudioModulation/Source/AudioModulationEditor/Private/Insights/Messages/ControlBusTraceMessages.h
+++ b/Engine/Plugins/Runtime/AudioModulation/Source/AudioModulationEditor/Private/Insights/Messages/ControlBusTraceMessages.h
@@ -53,8 +53,12 @@ namespace AudioModulationEditor
 		{
 			const UE::Trace::IAnalyzer::FEventData& EventData = InContext.EventData;
 			Value = EventData.GetValue<float>("Value");
+			EventData.GetString("Name", BusName);
+			EventData.GetString("ParamName", ParamName);
 		}
 
+		FString BusName;
+		FString ParamName;
 		float Value = 1.0f;
 	};
 
diff --git a/Engine/Plugins/Runtime/AudioModulation/Source/AudioModulationEditor/Private/Insights/Providers/ControlBusTraceProvider.cpp b/Engine/Plugins/Runtime/AudioModulation/Source/AudioModulationEditor/Private/Insights/Providers/ControlBusTraceProvider.cpp
--- a/Engine/Plugins/Runtime/AudioModulation/Source/AudioModulationEditor/Private/Insights/Providers/ControlBusTraceProvider.cpp
+++ b/Engine/Plugins/Runtime/AudioModulation/Source/AudioModulationEditor/Private/Insights/Providers/ControlBusTraceProvider.cpp
@@ -45,6 +45,8 @@ namespace AudioModulationEditor
 		[](const FControlBusUpdateMessage& Msg, TSharedPtr<FControlBusDashboardEntry>* OutEntry)
 		{
 			FControlBusDashboardEntry& EntryRef = *OutEntry->Get();
+			EntryRef.Name = *Msg.BusName;
+			EntryRef.ParamName = Msg.ParamName;
 			EntryRef.Value = Msg.Value;
 		});
 
