Fix serialization crash when attempting to access null config file.

diff --git a/Engine/Plugins/Runtime/AudioModulation/Source/AudioModulation/Private/AudioModulationProfileSerializer.h b/Engine/Plugins/Runtime/AudioModulation/Source/AudioModulation/Private/AudioModulationProfileSerializer.h
--- a/Engine/Plugins/Runtime/AudioModulation/Source/AudioModulation/Private/AudioModulationProfileSerializer.h
+++ b/Engine/Plugins/Runtime/AudioModulation/Source/AudioModulation/Private/AudioModulationProfileSerializer.h
@@ -109,20 +109,25 @@ namespace AudioModulation
 			}
 
 			// this will create and try to read in if not already loaded
-			FConfigFile& ConfigFile = *GConfig->Find(ConfigPath);
-			ConfigFile.Reset();
+			FConfigFile* ConfigFile = GConfig->Find(ConfigPath);
+			if (!ConfigFile)
+			{
+				ConfigFile = &(GConfig->Add(ConfigPath, FConfigFile()));
+			}
+
+			ConfigFile->Reset();
 
 			for (const FSoundControlBusMixStage& Stage : InBusMix.MixStages)
 			{
 				if (Stage.Bus)
 				{
 					FString SectionName = Stage.Bus->GetPathName();
-					ConfigFile.SetFloat(*SectionName, *GET_MEMBER_NAME_CHECKED(FSoundModulationMixValue, AttackTime).ToString(), Stage.Value.AttackTime);
-					ConfigFile.SetFloat(*SectionName, *GET_MEMBER_NAME_CHECKED(FSoundModulationMixValue, ReleaseTime).ToString(), Stage.Value.ReleaseTime);
-					ConfigFile.SetFloat(*SectionName, *GET_MEMBER_NAME_CHECKED(FSoundModulationMixValue, TargetValue).ToString(), Stage.Value.TargetValue);
+					ConfigFile->SetFloat(*SectionName, *GET_MEMBER_NAME_CHECKED(FSoundModulationMixValue, AttackTime).ToString(), Stage.Value.AttackTime);
+					ConfigFile->SetFloat(*SectionName, *GET_MEMBER_NAME_CHECKED(FSoundModulationMixValue, ReleaseTime).ToString(), Stage.Value.ReleaseTime);
+					ConfigFile->SetFloat(*SectionName, *GET_MEMBER_NAME_CHECKED(FSoundModulationMixValue, TargetValue).ToString(), Stage.Value.TargetValue);
 				}
 			}
-			ConfigFile.Dirty = true;
+			ConfigFile->Dirty = true;
 
 			UE_LOG(LogAudioModulation, Display, TEXT("Serialized mix to '%s'"), *Path);
 			GConfig->Flush(false, ConfigPath);
