[Audio Modulation Insights] AudioModulationEditor: Added the classic "if (!IsRunningCommandlet())" so cooking (or any other commandlet operation) doesn't fail.

diff --git a/Engine/Plugins/Runtime/AudioModulation/Source/AudioModulationEditor/Private/AudioModulationEditor.cpp b/Engine/Plugins/Runtime/AudioModulation/Source/AudioModulationEditor/Private/AudioModulationEditor.cpp
--- a/Engine/Plugins/Runtime/AudioModulation/Source/AudioModulationEditor/Private/AudioModulationEditor.cpp
+++ b/Engine/Plugins/Runtime/AudioModulation/Source/AudioModulationEditor/Private/AudioModulationEditor.cpp
@@ -58,13 +58,16 @@ namespace
 
 	void RegisterModulationInAudioInsights()
 	{
-		if (!bIsModulationRegisteredInAudioInsights)
+		// Don't run providers in any commandlet to avoid additional, unnecessary overhead as audio modulation insights is dormant.
+		if (!IsRunningCommandlet())
 		{
+			if (!bIsModulationRegisteredInAudioInsights)
+			{
+				IAudioInsightsEditorModule& InsightsEditorModule = IAudioInsightsEditorModule::GetChecked();
+				InsightsEditorModule.RegisterDashboardViewFactory(MakeShared<AudioModulationEditor::FControlBusDashboardViewFactory>());
 
-			IAudioInsightsEditorModule& InsightsModule = IAudioInsightsEditorModule::GetChecked();
-			InsightsModule.RegisterDashboardViewFactory(MakeShared<AudioModulationEditor::FControlBusDashboardViewFactory>());
-
-			bIsModulationRegisteredInAudioInsights = true;
+				bIsModulationRegisteredInAudioInsights = true;
+			}
 		}
 	}
 
