[Audio Modulation Insights] Implemented Modulation Matrix dashboard for Audio Insights (backend)

diff --git a/Engine/Plugins/Runtime/AudioModulation/Source/AudioModulation/Private/AudioModulationSystem.cpp b/Engine/Plugins/Runtime/AudioModulation/Source/AudioModulation/Private/AudioModulationSystem.cpp
--- a/Engine/Plugins/Runtime/AudioModulation/Source/AudioModulation/Private/AudioModulationSystem.cpp
+++ b/Engine/Plugins/Runtime/AudioModulation/Source/AudioModulation/Private/AudioModulationSystem.cpp
@@ -51,6 +51,51 @@ UE_TRACE_EVENT_BEGIN(Audio, ControlBusDeactivate)
 	UE_TRACE_EVENT_FIELD(uint32, ControlBusId)
 	UE_TRACE_EVENT_FIELD(double, Timestamp)
 UE_TRACE_EVENT_END()
+
+UE_TRACE_EVENT_BEGIN(Audio, BusMixRegisterBus)
+	UE_TRACE_EVENT_FIELD(uint32, DeviceId)
+	UE_TRACE_EVENT_FIELD(uint32, SourceId)
+	UE_TRACE_EVENT_FIELD(double, Timestamp)
+	UE_TRACE_EVENT_FIELD(uint32, ModulatingSourceId)
+	UE_TRACE_EVENT_FIELD(UE::Trace::WideString, BusName)
+UE_TRACE_EVENT_END()
+
+UE_TRACE_EVENT_BEGIN(Audio, BusMixActivate)
+	UE_TRACE_EVENT_FIELD(uint32, DeviceId)
+	UE_TRACE_EVENT_FIELD(uint32, SourceId)
+	UE_TRACE_EVENT_FIELD(double, Timestamp)
+	UE_TRACE_EVENT_FIELD(UE::Trace::WideString, Name)
+UE_TRACE_EVENT_END()
+
+UE_TRACE_EVENT_BEGIN(Audio, BusMixUpdate)
+	UE_TRACE_EVENT_FIELD(uint32, DeviceId)
+	UE_TRACE_EVENT_FIELD(uint32, SourceId)
+	UE_TRACE_EVENT_FIELD(double, Timestamp)
+	UE_TRACE_EVENT_FIELD(uint32[], BusIds)
+	UE_TRACE_EVENT_FIELD(float[], BusValues)
+UE_TRACE_EVENT_END()
+
+UE_TRACE_EVENT_BEGIN(Audio, GeneratorUpdate)
+	UE_TRACE_EVENT_FIELD(uint32, DeviceId)
+	UE_TRACE_EVENT_FIELD(uint32, SourceId)
+	UE_TRACE_EVENT_FIELD(double, Timestamp)
+	UE_TRACE_EVENT_FIELD(uint32[], BusIds)
+	UE_TRACE_EVENT_FIELD(float, GeneratorValue)
+UE_TRACE_EVENT_END()
+
+UE_TRACE_EVENT_BEGIN(Audio, BusFinalValuesUpdate)
+	UE_TRACE_EVENT_FIELD(uint32, DeviceId)
+	UE_TRACE_EVENT_FIELD(uint32, SourceId)
+	UE_TRACE_EVENT_FIELD(double, Timestamp)
+	UE_TRACE_EVENT_FIELD(uint32[], BusIds)
+	UE_TRACE_EVENT_FIELD(float[], BusValues)
+UE_TRACE_EVENT_END()
+
+UE_TRACE_EVENT_BEGIN(Audio, ModulatingSourceDeactivate)
+	UE_TRACE_EVENT_FIELD(uint32, DeviceId)
+	UE_TRACE_EVENT_FIELD(uint32, SourceId)
+	UE_TRACE_EVENT_FIELD(double, Timestamp)
+UE_TRACE_EVENT_END()
 #endif // UE_AUDIO_PROFILERTRACE_ENABLED
 
 namespace AudioModulation
@@ -124,6 +169,24 @@ namespace AudioModulation
 			}
 
 			ManuallyActivatedBusMixes.Add(MoveTemp(BusMixHandle));
+
+#if UE_AUDIO_PROFILERTRACE_ENABLED
+			for (const FModulatorBusMixStageSettings& StageSetting : Settings.Stages)
+			{
+				UE_TRACE_LOG(Audio, BusMixRegisterBus, AudioChannel)
+					<< BusMixRegisterBus.DeviceId(AudioDeviceId)
+					<< BusMixRegisterBus.SourceId(StageSetting.BusSettings.GetId())
+					<< BusMixRegisterBus.Timestamp(FPlatformTime::Cycles64())
+					<< BusMixRegisterBus.ModulatingSourceId(Settings.GetId())
+					<< BusMixRegisterBus.BusName(*(StageSetting.BusSettings.GetName().ToString()));
+			}
+
+			UE_TRACE_LOG(Audio, BusMixActivate, AudioChannel)
+				<< BusMixActivate.DeviceId(AudioDeviceId)
+				<< BusMixActivate.SourceId(Settings.GetId())
+				<< BusMixActivate.Timestamp(FPlatformTime::Cycles64())
+				<< BusMixActivate.Name(*Settings.GetName().ToString());
+#endif
 		});
 	}
 
@@ -200,6 +263,13 @@ namespace AudioModulation
 			{
 				FModulatorBusMixProxy& MixProxy = MixHandle.FindProxy();
 				MixProxy.SetStopping();
+
+#if UE_AUDIO_PROFILERTRACE_ENABLED
+				UE_TRACE_LOG(Audio, ModulatingSourceDeactivate, AudioChannel)
+					<< ModulatingSourceDeactivate.DeviceId(AudioDeviceId)
+					<< ModulatingSourceDeactivate.SourceId(BusMixId)
+					<< ModulatingSourceDeactivate.Timestamp(FPlatformTime::Cycles64());
+#endif
 			}
 		});
 	}
@@ -212,9 +282,16 @@ namespace AudioModulation
 
 		RunCommandOnProcessingThread([this]()
 		{
-			for (TPair<FBusMixId, FModulatorBusMixProxy>& Pair : RefProxies.BusMixes)
+			for (auto& [BusMixId, MixProxy] : RefProxies.BusMixes)
 			{
-				Pair.Value.SetStopping();
+				MixProxy.SetStopping();
+
+#if UE_AUDIO_PROFILERTRACE_ENABLED
+				UE_TRACE_LOG(Audio, ModulatingSourceDeactivate, AudioChannel)
+					<< ModulatingSourceDeactivate.DeviceId(AudioDeviceId)
+					<< ModulatingSourceDeactivate.SourceId(BusMixId)
+					<< ModulatingSourceDeactivate.Timestamp(FPlatformTime::Cycles64());
+#endif
 			}
 		});
 	}
@@ -569,16 +646,82 @@ namespace AudioModulation
 		const bool bChannelEnabled = UE_TRACE_CHANNELEXPR_IS_ENABLED(AudioChannel);
 		if (bChannelEnabled)
 		{
-			for (TPair<FBusId, FControlBusProxy>& Pair : RefProxies.Buses)
+			TMap<FGeneratorId, TSet<FBusId>> GeneratorIdToBusIdsMap;
+			TArray<FBusId, TInlineAllocator<32>> BusIdsFinalValues;
+			TArray<float, TInlineAllocator<32>> BusFinalValues;
+
+			const int32 NumBuses = RefProxies.Buses.Num();
+			BusIdsFinalValues.Reserve(NumBuses);
+			BusFinalValues.Reserve(NumBuses);
+
+			// Control Bus message trace
+			for (const auto& [BusId, ControlBusProxy] : RefProxies.Buses)
 			{
 				UE_TRACE_LOG(Audio, ControlBusUpdate, AudioChannel)
 					<< ControlBusUpdate.DeviceId(static_cast<uint32>(AudioDeviceId))
-					<< ControlBusUpdate.ControlBusId(static_cast<uint32>(Pair.Key))
+					<< ControlBusUpdate.ControlBusId(static_cast<uint32>(BusId))
 					<< ControlBusUpdate.Timestamp(FPlatformTime::Cycles64())
-					<< ControlBusUpdate.Name(*(Pair.Value.GetName().ToString()))
-					<< ControlBusUpdate.ParamName(*(Pair.Value.GetParameterName().ToString()))
-					<< ControlBusUpdate.Value(Pair.Value.GetValue());
+					<< ControlBusUpdate.Name(*(ControlBusProxy.GetName().ToString()))
+					<< ControlBusUpdate.ParamName(*(ControlBusProxy.GetParameterName().ToString()))
+					<< ControlBusUpdate.Value(ControlBusProxy.GetValue());
+
+				// Collect the final value of this control bus
+				BusIdsFinalValues.Emplace(BusId);
+				BusFinalValues.Emplace(ControlBusProxy.GetValue());
+
+				// Collect generators set in this control bus
+				for (const FGeneratorHandle& GeneratorHandle : ControlBusProxy.GetGeneratorHandles())
+				{
+					GeneratorIdToBusIdsMap.FindOrAdd(GeneratorHandle.GetId()).Emplace(BusId);
+				}
 			}
+
+			// Control Bus Mix message trace
+			for (const auto& [BusMixId, BusMixProxy] : RefProxies.BusMixes)
+			{
+				const int32 NumStages = BusMixProxy.Stages.Num();
+
+				TArray<FBusId, TInlineAllocator<32>> BusIds;
+				TArray<float,  TInlineAllocator<32>> BusMixCurrentValues;
+
+				BusIds.Reserve(NumStages);
+				BusMixCurrentValues.Reserve(NumStages);
+
+				for (const auto& [BusId, BusMixStageProxy] : BusMixProxy.Stages)
+				{
+					BusIds.Emplace(BusId);
+					BusMixCurrentValues.Emplace(BusMixStageProxy.Value.GetCurrentValue());
+				}
+
+				UE_TRACE_LOG(Audio, BusMixUpdate, AudioChannel)
+					<< BusMixUpdate.DeviceId(AudioDeviceId)
+					<< BusMixUpdate.SourceId(BusMixId)
+					<< BusMixUpdate.Timestamp(FPlatformTime::Cycles64())
+					<< BusMixUpdate.BusIds(BusIds.GetData(), BusIds.Num())
+					<< BusMixUpdate.BusValues(BusMixCurrentValues.GetData(), BusMixCurrentValues.Num());
+			}
+
+			// Generator message trace
+			for (const auto& [GeneratorId, BusIds] : GeneratorIdToBusIdsMap)
+			{
+				if (const FModulatorGeneratorProxy* FoundGeneratorProxyPtr = RefProxies.Generators.Find(GeneratorId))
+				{
+					UE_TRACE_LOG(Audio, GeneratorUpdate, AudioChannel)
+						<< GeneratorUpdate.DeviceId(AudioDeviceId)
+						<< GeneratorUpdate.SourceId(GeneratorId)
+						<< GeneratorUpdate.Timestamp(FPlatformTime::Cycles64())
+						<< GeneratorUpdate.BusIds(BusIds.Array().GetData(), BusIds.Num())
+						<< GeneratorUpdate.GeneratorValue(FoundGeneratorProxyPtr->GetValue());
+				}
+			}
+
+			// Bus Final values message trace
+			UE_TRACE_LOG(Audio, BusFinalValuesUpdate, AudioChannel)
+				<< BusFinalValuesUpdate.DeviceId(AudioDeviceId)
+				<< BusFinalValuesUpdate.SourceId(INDEX_NONE)
+				<< BusFinalValuesUpdate.Timestamp(FPlatformTime::Cycles64())
+				<< BusFinalValuesUpdate.BusIds(BusIdsFinalValues.GetData(), BusIdsFinalValues.Num())
+				<< BusFinalValuesUpdate.BusValues(BusFinalValues.GetData(), BusFinalValues.Num());
 		}
 #endif // UE_AUDIO_PROFILERTRACE_ENABLED
 
@@ -802,6 +945,12 @@ namespace AudioModulation
 			FGeneratorHandle GeneratorHandle = FGeneratorHandle::Get(static_cast<FGeneratorId>(ModId), RefProxies.Generators);
 			if (UnregisterModulator<FGeneratorHandle>(GeneratorHandle, RefModulators.GeneratorMap, HandleId))
 			{
+#if UE_AUDIO_PROFILERTRACE_ENABLED
+				UE_TRACE_LOG(Audio, ModulatingSourceDeactivate, AudioChannel)
+					<< ModulatingSourceDeactivate.DeviceId(AudioDeviceId)
+					<< ModulatingSourceDeactivate.SourceId(ModId)
+					<< ModulatingSourceDeactivate.Timestamp(FPlatformTime::Cycles64());
+#endif
 				return;
 			}
 		});
diff --git a/Engine/Plugins/Runtime/AudioModulation/Source/AudioModulation/Private/SoundControlBusProxy.cpp b/Engine/Plugins/Runtime/AudioModulation/Source/AudioModulation/Private/SoundControlBusProxy.cpp
--- a/Engine/Plugins/Runtime/AudioModulation/Source/AudioModulation/Private/SoundControlBusProxy.cpp
+++ b/Engine/Plugins/Runtime/AudioModulation/Source/AudioModulation/Private/SoundControlBusProxy.cpp
@@ -18,6 +18,22 @@ UE_TRACE_EVENT_BEGIN(Audio, ControlBusActivate)
 	UE_TRACE_EVENT_FIELD(UE::Trace::WideString, Name)
 	UE_TRACE_EVENT_FIELD(UE::Trace::WideString, ParamName)
 UE_TRACE_EVENT_END()
+
+UE_TRACE_EVENT_BEGIN(Audio, GeneratorRegisterBus)
+	UE_TRACE_EVENT_FIELD(uint32, DeviceId)
+	UE_TRACE_EVENT_FIELD(uint32, SourceId)
+	UE_TRACE_EVENT_FIELD(double, Timestamp)
+	UE_TRACE_EVENT_FIELD(uint32, ModulatingSourceId)
+	UE_TRACE_EVENT_FIELD(UE::Trace::WideString, BusName)
+UE_TRACE_EVENT_END()
+
+UE_TRACE_EVENT_BEGIN(Audio, GeneratorActivate)
+	UE_TRACE_EVENT_FIELD(uint32, DeviceId)
+	UE_TRACE_EVENT_FIELD(uint32, SourceId)
+	UE_TRACE_EVENT_FIELD(double, Timestamp)
+	UE_TRACE_EVENT_FIELD(UE::Trace::WideString, Name)
+UE_TRACE_EVENT_END()
+
 #endif // UE_AUDIO_PROFILERTRACE_ENABLED
 
 namespace AudioModulation
@@ -35,6 +51,22 @@ namespace AudioModulation
 			<< ControlBusActivate.Timestamp(FPlatformTime::Cycles64())
 			<< ControlBusActivate.Name(*GetName().ToString())
 			<< ControlBusActivate.ParamName(*OutputParameter.ParameterName.ToString());
+
+		for (const FModulationGeneratorSettings& GeneratorSetting : GeneratorSettings)
+		{
+			UE_TRACE_LOG(Audio, GeneratorRegisterBus, AudioChannel)
+				<< GeneratorRegisterBus.DeviceId(ModSystem.GetAudioDeviceId())
+				<< GeneratorRegisterBus.SourceId(GetId())
+				<< GeneratorRegisterBus.Timestamp(FPlatformTime::Cycles64())
+				<< GeneratorRegisterBus.ModulatingSourceId(GeneratorSetting.GetId())
+				<< GeneratorRegisterBus.BusName(*(GetName().ToString()));
+
+			UE_TRACE_LOG(Audio, GeneratorActivate, AudioChannel)
+				<< GeneratorActivate.DeviceId(ModSystem.GetAudioDeviceId())
+				<< GeneratorActivate.SourceId(GeneratorSetting.GetId())
+				<< GeneratorActivate.Timestamp(FPlatformTime::Cycles64())
+				<< GeneratorActivate.Name(*GeneratorSetting.GetName().ToString());
+		}
 #endif
 
 		return ModSystem.RegisterModulator(HandleId, *this);
diff --git a/Engine/Plugins/Runtime/AudioModulation/Source/AudioModulationEditor/Private/Insights/Messages/ModulationMatrixTraceMessages.h b/Engine/Plugins/Runtime/AudioModulation/Source/AudioModulationEditor/Private/Insights/Messages/ModulationMatrixTraceMessages.h
new file mode 100644
--- /dev/null
+++ b/Engine/Plugins/Runtime/AudioModulation/Source/AudioModulationEditor/Private/Insights/Messages/ModulationMatrixTraceMessages.h
@@ -0,0 +1,190 @@
+// Copyright Epic Games, Inc. All Rights Reserved.
+#pragma once
+
+#include "AudioDefines.h"
+#include "Messages/AnalyzerMessageQueue.h"
+#include "Trace/Analyzer.h"
+#include "Views/TableDashboardViewFactory.h"
+
+namespace AudioModulationEditor
+{
+	using FBusId    = uint32;
+	using FSourceId = uint32;
+
+	enum class EModulationMatrixEntryType : uint8
+	{
+		// Arranged alphabetically to simplify sorting
+		BusMix,
+		BusFinalValues,
+		Generator,
+		None
+	};
+
+
+	// Modulation matrix dashboard entry
+	class FModulationMatrixDashboardEntry : public UE::Audio::Insights::FSoundAssetDashboardEntry
+	{
+	public:
+		FModulationMatrixDashboardEntry() = default;
+		virtual ~FModulationMatrixDashboardEntry() = default;
+
+		FSourceId SourceId = INDEX_NONE;
+		EModulationMatrixEntryType EntryType = EModulationMatrixEntryType::None;
+
+		TMap<FBusId, float> BusIdToValueMap;
+	};
+	
+
+	// Modulation matrix message base
+	struct FModulationMatrixMessageBase
+	{
+		FModulationMatrixMessageBase() = default;
+
+		FModulationMatrixMessageBase(const UE::Trace::IAnalyzer::FOnEventContext& InContext)
+		{
+			const UE::Trace::IAnalyzer::FEventData& EventData = InContext.EventData;
+
+			DeviceId  = static_cast<::Audio::FDeviceId>(EventData.GetValue<uint32>("DeviceId"));
+			Timestamp = InContext.EventTime.AsSeconds(EventData.GetValue<uint64>("Timestamp"));
+			SourceId  = static_cast<FSourceId>(EventData.GetValue<uint32>("SourceId"));
+		}
+
+		Audio::FDeviceId DeviceId    = INDEX_NONE;
+		FSourceId SourceId = INDEX_NONE;
+		double Timestamp = 0.0;
+	};
+
+
+	// Register bus message
+	struct FModulationMatrixRegisterBusMessage : public FModulationMatrixMessageBase
+	{
+		FModulationMatrixRegisterBusMessage() = default;
+
+		FModulationMatrixRegisterBusMessage(const UE::Trace::IAnalyzer::FOnEventContext& InContext)
+			: FModulationMatrixMessageBase(InContext)
+		{
+			const UE::Trace::IAnalyzer::FEventData& EventData = InContext.EventData;
+
+			ModulatingSourceId = static_cast<FSourceId>(EventData.GetValue<uint32>("ModulatingSourceId"));
+			EventData.GetString("BusName", BusName);
+		}
+
+		FSourceId ModulatingSourceId = INDEX_NONE;
+		FString BusName;
+	};
+
+
+	// Activate messages
+	struct FModulationMatrixActivateMessage : public FModulationMatrixMessageBase
+	{
+		FModulationMatrixActivateMessage() = default;
+
+		FModulationMatrixActivateMessage(const UE::Trace::IAnalyzer::FOnEventContext& InContext)
+			: FModulationMatrixMessageBase(InContext)
+		{
+			const UE::Trace::IAnalyzer::FEventData& EventData = InContext.EventData;
+
+			EventData.GetString("Name", Name);
+		}
+
+		EModulationMatrixEntryType EntryType = EModulationMatrixEntryType::None;
+		FString Name;
+	};
+
+	struct FBusMixActivateMessage : public FModulationMatrixActivateMessage
+	{
+		FBusMixActivateMessage() = default;
+
+		FBusMixActivateMessage(const UE::Trace::IAnalyzer::FOnEventContext& InContext)
+			: FModulationMatrixActivateMessage(InContext)
+		{
+			EntryType = EModulationMatrixEntryType::BusMix;
+		}
+	};
+
+	struct FGeneratorActivateMessage : public FModulationMatrixActivateMessage
+	{
+		FGeneratorActivateMessage() = default;
+
+		FGeneratorActivateMessage(const UE::Trace::IAnalyzer::FOnEventContext& InContext)
+			: FModulationMatrixActivateMessage(InContext)
+		{
+			EntryType = EModulationMatrixEntryType::Generator;
+		}
+	};
+
+
+	// Update messages
+	struct FModulationMatrixUpdateMessage : public FModulationMatrixMessageBase
+	{
+		FModulationMatrixUpdateMessage() = default;
+
+		FModulationMatrixUpdateMessage(const UE::Trace::IAnalyzer::FOnEventContext& InContext)
+			: FModulationMatrixMessageBase(InContext)
+		{
+			const UE::Trace::IAnalyzer::FEventData& EventData = InContext.EventData;
+
+			const TArrayView<const FBusId> BusIds   = EventData.GetArrayView<FBusId>("BusIds");
+			const TArrayView<const float> BusValues = EventData.GetArrayView<float>("BusValues");
+
+			check(BusIds.Num() == BusValues.Num());
+
+			for (int32 Index = 0; Index < BusIds.Num(); ++Index)
+			{
+				BusIdToValueMap.Emplace(BusIds[Index], BusValues[Index]);
+			}
+		}
+
+		TMap<FBusId, float> BusIdToValueMap;
+	};
+
+	using FBusMixUpdateMessage = FModulationMatrixUpdateMessage;
+
+	struct FGeneratorUpdateMessage : public FModulationMatrixMessageBase
+	{
+		FGeneratorUpdateMessage() = default;
+
+		FGeneratorUpdateMessage(const UE::Trace::IAnalyzer::FOnEventContext& InContext)
+			: FModulationMatrixMessageBase(InContext)
+		{
+			const UE::Trace::IAnalyzer::FEventData& EventData = InContext.EventData;
+
+			const TArrayView<const FBusId> BusIds = EventData.GetArrayView<FBusId>("BusIds");
+			const float GeneratorValue = EventData.GetValue<float>("GeneratorValue");
+
+			// All control buses use the same generator value
+			for (int32 Index = 0; Index < BusIds.Num(); ++Index)
+			{
+				BusIdToValueMap.Emplace(BusIds[Index], GeneratorValue);
+			}
+		}
+
+		TMap<FBusId, float> BusIdToValueMap;
+	};
+
+	using FBusFinalValuesUpdateMessage = FModulationMatrixUpdateMessage;
+
+
+	// Deactivate message
+	using FModulationMatrixDeactivateMessage = FModulationMatrixMessageBase;
+
+
+	// Modulation matrix messages
+	class FModulationMatrixMessages
+	{
+		UE::Audio::Insights::TAnalyzerMessageQueue<FModulationMatrixRegisterBusMessage> RegisterBusMessages{ 0.1 };
+
+		UE::Audio::Insights::TAnalyzerMessageQueue<FBusMixActivateMessage> BusMixActivateMessages{ 0.1 };
+		UE::Audio::Insights::TAnalyzerMessageQueue<FBusMixUpdateMessage> BusMixUpdateMessages{ 0.1 };
+
+		UE::Audio::Insights::TAnalyzerMessageQueue<FGeneratorActivateMessage> GeneratorActivateMessages{ 0.1 };
+		UE::Audio::Insights::TAnalyzerMessageQueue<FGeneratorUpdateMessage> GeneratorUpdateMessages{ 0.1 };
+
+		UE::Audio::Insights::TAnalyzerMessageQueue<FBusFinalValuesUpdateMessage> BusFinalValuesUpdateMessages{ 0.1 };
+
+		UE::Audio::Insights::TAnalyzerMessageQueue<FModulationMatrixDeactivateMessage> DeactivateMessages{ 0.1 };
+
+		friend class FModulationMatrixTraceProvider;
+	};
+
+} // namespace AudioModulationEditor
diff --git a/Engine/Plugins/Runtime/AudioModulation/Source/AudioModulationEditor/Private/Insights/Providers/ModulationMatrixTraceProvider.cpp b/Engine/Plugins/Runtime/AudioModulation/Source/AudioModulationEditor/Private/Insights/Providers/ModulationMatrixTraceProvider.cpp
new file mode 100644
--- /dev/null
+++ b/Engine/Plugins/Runtime/AudioModulation/Source/AudioModulationEditor/Private/Insights/Providers/ModulationMatrixTraceProvider.cpp
@@ -0,0 +1,354 @@
+// Copyright Epic Games, Inc. All Rights Reserved.
+#include "Insights/Providers/ModulationMatrixTraceProvider.h"
+
+#include "AudioDeviceManager.h"
+#include "Trace/Analyzer.h"
+#include "TraceServices/Model/AnalysisSession.h"
+
+namespace AudioModulationEditor
+{
+	FModulationMatrixTraceProvider::FModulationMatrixTraceProvider(TSharedRef<UE::Audio::Insights::FTraceChannelManager> InManager)
+		: UE::Audio::Insights::TDeviceDataMapTraceProvider<uint32, TSharedPtr<FModulationMatrixDashboardEntry>>(GetName_Static())
+	{
+		Channels.Add(InManager->CreateHandle({ TEXT("AudioChannel") }));
+
+		FAudioDeviceManagerDelegates::OnAudioDeviceDestroyed.AddRaw(this, &FModulationMatrixTraceProvider::OnAudioDeviceDestroyed);
+	}
+
+	FModulationMatrixTraceProvider::~FModulationMatrixTraceProvider()
+	{
+		FAudioDeviceManagerDelegates::OnAudioDeviceDestroyed.RemoveAll(this);
+	}
+
+	FName FModulationMatrixTraceProvider::GetName_Static()
+	{
+		static const FName ModulationMatrixProviderName = "ModulationMatrixProvider";
+		return ModulationMatrixProviderName;
+	}
+
+	void FModulationMatrixTraceProvider::UpdateActiveControlBusesToAdd(const TMap<FBusId, FString>& InBusIdToBusNameMap)
+	{
+		for (const auto& [BusId, BusName] : InBusIdToBusNameMap)
+		{
+			if (BusInfo* FoundBusInfoPtr = ActiveControlBuses.Find(BusId))
+			{
+				BusInfo& FoundBusInfo = *FoundBusInfoPtr;
+				FoundBusInfo.RefCount++;
+			}
+			else
+			{
+				ActiveControlBuses.Emplace(BusId, { BusName, 1 });
+			}
+		}
+
+		ActiveControlBuses.ValueSort([](const BusInfo& A, const BusInfo& B)
+		{
+			return A.BusName.ToLower() < B.BusName.ToLower();
+		});
+	}
+
+	void FModulationMatrixTraceProvider::UpdateActiveControlBusesToRemove(const TMap<FBusId, float>& InBusIdToValueMap)
+	{
+		TArray<FBusId> BusesPendingToRemove;
+
+		for (const auto& [BusId, BusValue] : InBusIdToValueMap)
+		{
+			if (BusInfo* FoundBusInfoPtr = ActiveControlBuses.Find(BusId))
+			{
+				BusInfo& FoundBusInfo = *FoundBusInfoPtr;
+				FoundBusInfo.RefCount--;
+
+				if (FoundBusInfo.RefCount <= 0)
+				{
+					RemovedControlBusesNames.Emplace(FoundBusInfo.BusName);
+					BusesPendingToRemove.Emplace(BusId);
+				}
+			}
+		}
+
+		for (const FBusId& BusId : BusesPendingToRemove)
+		{
+			ActiveControlBuses.Remove(BusId);
+		}
+	}
+
+	bool FModulationMatrixTraceProvider::ProcessMessages()
+	{
+		// Helper lambdas
+		auto BumpEntryFunc = [this](const FModulationMatrixMessageBase& Msg)
+		{
+			TSharedPtr<FModulationMatrixDashboardEntry>* ToReturn = nullptr;
+
+			UpdateDeviceEntry(Msg.DeviceId, Msg.SourceId, [&ToReturn, &Msg](TSharedPtr<FModulationMatrixDashboardEntry>& Entry)
+			{
+				if (!Entry.IsValid())
+				{
+					Entry = MakeShared<FModulationMatrixDashboardEntry>();
+					Entry->DeviceId = Msg.DeviceId;
+					Entry->SourceId = Msg.SourceId;
+				}
+
+				Entry->Timestamp = Msg.Timestamp;
+
+				ToReturn = &Entry;
+			});
+
+			return ToReturn;
+		};
+
+		auto GetEntry = [this](const FModulationMatrixMessageBase& Msg)
+		{
+			return FindDeviceEntry(Msg.DeviceId, Msg.SourceId);
+		};
+
+		// Process messages
+		uint32 NumActiveBuses = ActiveControlBuses.Num();
+
+		if (!(TraceMessages.RegisterBusMessages.IsEmpty() && (TraceMessages.BusMixActivateMessages.IsEmpty() || TraceMessages.GeneratorActivateMessages.IsEmpty())))
+		{
+			TMap<FSourceId, TMap<FBusId, FString>> ModulatingSourceIdToBusesInfoMap;
+
+			auto UpdateActiveModulatorSourceIds = [&](const Audio::FDeviceId DeviceId, const FSourceId SourceId)
+			{
+				TSet<uint32>& ActiveModulatorSourceIds = DeviceIdToActiveModulatorSourceIdsMap.FindOrAdd(DeviceId);
+
+				if (!ActiveModulatorSourceIds.Contains(SourceId))
+				{
+					ActiveModulatorSourceIds.Add(SourceId);
+
+					if (const TMap<FBusId, FString>* FoundBusInfoPtr = ModulatingSourceIdToBusesInfoMap.Find(SourceId))
+					{
+						UpdateActiveControlBusesToAdd(*FoundBusInfoPtr);
+					}
+				}
+			};
+
+			ProcessMessageQueue<FModulationMatrixRegisterBusMessage>(TraceMessages.RegisterBusMessages,
+			[](const FModulationMatrixMessageBase& Msg)	{ return nullptr; }, // No entry operation
+			[this, &ModulatingSourceIdToBusesInfoMap](const FModulationMatrixRegisterBusMessage& Msg, TSharedPtr<FModulationMatrixDashboardEntry>* OutEntry)
+			{
+				if (TMap<FBusId, FString>* FoundBusInfoPtr = ModulatingSourceIdToBusesInfoMap.Find(Msg.ModulatingSourceId))
+				{
+					FoundBusInfoPtr->FindOrAdd(Msg.SourceId, Msg.BusName);
+				}
+				else
+				{
+					ModulatingSourceIdToBusesInfoMap.Emplace(Msg.ModulatingSourceId).Emplace(Msg.SourceId, Msg.BusName);
+				}
+			});
+
+			ProcessMessageQueue<FBusMixActivateMessage>(TraceMessages.BusMixActivateMessages, BumpEntryFunc,
+			[this, &UpdateActiveModulatorSourceIds](const FBusMixActivateMessage& Msg, TSharedPtr<FModulationMatrixDashboardEntry>* OutEntry)
+			{
+				FModulationMatrixDashboardEntry& EntryRef = *OutEntry->Get();
+				EntryRef.Name = *Msg.Name;
+				EntryRef.EntryType = EModulationMatrixEntryType::BusMix;
+
+				UpdateActiveModulatorSourceIds(Msg.DeviceId, EntryRef.SourceId);
+			});
+
+			ProcessMessageQueue<FGeneratorActivateMessage>(TraceMessages.GeneratorActivateMessages, BumpEntryFunc,
+			[this, &UpdateActiveModulatorSourceIds](const FGeneratorActivateMessage& Msg, TSharedPtr<FModulationMatrixDashboardEntry>* OutEntry)
+			{
+				FModulationMatrixDashboardEntry& EntryRef = *OutEntry->Get();
+				EntryRef.Name = *Msg.Name;
+				EntryRef.EntryType = EModulationMatrixEntryType::Generator;
+
+				UpdateActiveModulatorSourceIds(Msg.DeviceId, EntryRef.SourceId);
+			});
+		}
+
+		ProcessMessageQueue<FModulationMatrixDeactivateMessage>(TraceMessages.DeactivateMessages, GetEntry,
+		[this](const FModulationMatrixDeactivateMessage& Msg, TSharedPtr<FModulationMatrixDashboardEntry>* OutEntry)
+		{
+			if (OutEntry && (*OutEntry)->Timestamp < Msg.Timestamp)
+			{
+				const FModulationMatrixDashboardEntry& EntryRef = *OutEntry->Get();
+				UpdateActiveControlBusesToRemove(EntryRef.BusIdToValueMap);
+
+				if (DeviceIdToActiveModulatorSourceIdsMap.Contains(Msg.DeviceId))
+				{
+					DeviceIdToActiveModulatorSourceIdsMap[Msg.DeviceId].Remove(EntryRef.SourceId);
+				}
+
+				RemoveDeviceEntry(Msg.DeviceId, Msg.SourceId);
+			}
+		});
+
+		ProcessMessageQueue<FBusMixUpdateMessage>(TraceMessages.BusMixUpdateMessages, GetEntry,
+		[this](const FBusMixUpdateMessage& Msg, TSharedPtr<FModulationMatrixDashboardEntry>* OutEntry)
+		{
+			if (OutEntry)
+			{
+				FModulationMatrixDashboardEntry& EntryRef = *OutEntry->Get();
+				EntryRef.BusIdToValueMap = Msg.BusIdToValueMap;
+				EntryRef.Timestamp = Msg.Timestamp;
+			}
+		});
+
+		ProcessMessageQueue<FGeneratorUpdateMessage>(TraceMessages.GeneratorUpdateMessages, GetEntry,
+		[this](const FGeneratorUpdateMessage& Msg, TSharedPtr<FModulationMatrixDashboardEntry>* OutEntry)
+		{
+			if (OutEntry)
+			{
+				FModulationMatrixDashboardEntry& EntryRef = *OutEntry->Get();
+				EntryRef.BusIdToValueMap = Msg.BusIdToValueMap;
+				EntryRef.Timestamp = Msg.Timestamp;
+			}
+		});
+
+		ProcessMessageQueue<FBusFinalValuesUpdateMessage>(TraceMessages.BusFinalValuesUpdateMessages, BumpEntryFunc,
+		[this](const FBusFinalValuesUpdateMessage& Msg, TSharedPtr<FModulationMatrixDashboardEntry>* OutEntry)
+		{
+			FModulationMatrixDashboardEntry& EntryRef = *OutEntry->Get();
+
+			if (EntryRef.Name.IsEmpty())
+			{
+				static const FString BusFinalValuesEntryName("Final Values: ");
+				EntryRef.Name      = BusFinalValuesEntryName;
+				EntryRef.EntryType = EModulationMatrixEntryType::BusFinalValues;
+			}
+
+			EntryRef.BusIdToValueMap = Msg.BusIdToValueMap;
+		});
+
+		// Notify Control Buses update
+		if (NumActiveBuses != ActiveControlBuses.Num() && !ActiveControlBuses.IsEmpty())
+		{
+			OnControlBusesAdded.ExecuteIfBound(ActiveControlBuses);
+		}
+
+		if (!RemovedControlBusesNames.IsEmpty())
+		{
+			OnControlBusesRemoved.ExecuteIfBound(RemovedControlBusesNames);
+			RemovedControlBusesNames.Reset();
+		}
+
+		return true;
+	}
+
+	void FModulationMatrixTraceProvider::OnAudioDeviceDestroyed(::Audio::FDeviceId InDeviceId)
+	{
+		DeviceIdToActiveModulatorSourceIdsMap.Reset();
+		ActiveControlBuses.Reset();
+	}
+
+	UE::Trace::IAnalyzer* FModulationMatrixTraceProvider::ConstructAnalyzer(TraceServices::IAnalysisSession& InSession)
+	{
+		class FModulationMatrixTraceAnalyzer : public UE::Audio::Insights::FTraceProviderBase::FTraceAnalyzerBase
+		{
+		public:
+			FModulationMatrixTraceAnalyzer(TSharedRef<FModulationMatrixTraceProvider> InProvider, TraceServices::IAnalysisSession& InSession)
+				: UE::Audio::Insights::FTraceProviderBase::FTraceAnalyzerBase(InProvider)
+				, Session(InSession)
+			{
+			}
+
+			virtual void OnAnalysisBegin(const UE::Trace::IAnalyzer::FOnAnalysisContext& Context) override
+			{
+				UE::Audio::Insights::FTraceProviderBase::FTraceAnalyzerBase::OnAnalysisBegin(Context);
+
+				UE::Trace::IAnalyzer::FInterfaceBuilder& Builder = Context.InterfaceBuilder;
+
+				Builder.RouteEvent(RouteId_BusMixRegisterBus,    "Audio", "BusMixRegisterBus");
+				Builder.RouteEvent(RouteId_GeneratorRegisterBus, "Audio", "GeneratorRegisterBus");
+
+				Builder.RouteEvent(RouteId_BusMixActivate,       "Audio", "BusMixActivate");
+				Builder.RouteEvent(RouteId_GeneratorActivate,    "Audio", "GeneratorActivate");
+
+				Builder.RouteEvent(RouteId_BusMixUpdate,         "Audio", "BusMixUpdate");
+				Builder.RouteEvent(RouteId_GeneratorUpdate,      "Audio", "GeneratorUpdate");
+				Builder.RouteEvent(RouteId_BusFinalValuesUpdate, "Audio", "BusFinalValuesUpdate");
+
+				Builder.RouteEvent(RouteId_Deactivate, "Audio", "ModulatingSourceDeactivate");
+			}
+
+			virtual bool OnEvent(uint16 RouteId, UE::Trace::IAnalyzer::EStyle Style, const UE::Trace::IAnalyzer::FOnEventContext& Context) override
+			{
+				LLM_SCOPE_BYNAME(TEXT("Insights/FModulationMatrixTraceAnalyzer"));
+
+				FModulationMatrixMessages& Messages = GetProvider<FModulationMatrixTraceProvider>().TraceMessages;
+
+				switch (RouteId)
+				{
+					case RouteId_BusMixRegisterBus:
+					case RouteId_GeneratorRegisterBus:
+					{
+						Messages.RegisterBusMessages.Enqueue(FModulationMatrixRegisterBusMessage{ Context });
+						break;
+					}
+
+					case RouteId_BusMixActivate:
+					{
+						Messages.BusMixActivateMessages.Enqueue(FBusMixActivateMessage{ Context });
+						break;
+					}
+
+					case RouteId_GeneratorActivate:
+					{
+						Messages.GeneratorActivateMessages.Enqueue(FGeneratorActivateMessage{ Context });
+						break;
+					}
+
+					case RouteId_BusMixUpdate:
+					{
+						Messages.BusMixUpdateMessages.Enqueue(FBusMixUpdateMessage{ Context });
+						break;
+					}
+
+					case RouteId_GeneratorUpdate:
+					{
+						Messages.GeneratorUpdateMessages.Enqueue(FGeneratorUpdateMessage{ Context });
+						break;
+					}
+
+					case RouteId_BusFinalValuesUpdate:
+					{
+						Messages.BusFinalValuesUpdateMessages.Enqueue(FBusFinalValuesUpdateMessage{ Context });
+						break;
+					}
+
+					case RouteId_Deactivate:
+					{
+						Messages.DeactivateMessages.Enqueue(FModulationMatrixDeactivateMessage{ Context });
+						break;
+					}
+
+					default:
+					{
+						return OnEventFailure(RouteId, Style, Context);
+					}
+				}
+
+				const double Timestamp = Context.EventTime.AsSeconds(Context.EventData.GetValue<uint64>("Timestamp"));
+
+				{
+					TraceServices::FAnalysisSessionEditScope SessionEditScope(Session);
+					Session.UpdateDurationSeconds(Timestamp);
+				}
+
+				return OnEventSuccess(RouteId, Style, Context);
+			}
+
+		private:
+			enum : uint16
+			{
+				RouteId_BusMixRegisterBus,
+				RouteId_GeneratorRegisterBus,
+
+				RouteId_BusMixActivate,
+				RouteId_GeneratorActivate,
+
+				RouteId_BusMixUpdate,
+				RouteId_GeneratorUpdate,
+				RouteId_BusFinalValuesUpdate,
+
+				RouteId_Deactivate
+			};
+
+			TraceServices::IAnalysisSession& Session;
+		};
+
+		return new FModulationMatrixTraceAnalyzer(AsShared(), InSession);
+	}
+} // namespace AudioModulationEditor
diff --git a/Engine/Plugins/Runtime/AudioModulation/Source/AudioModulationEditor/Private/Insights/Providers/ModulationMatrixTraceProvider.h b/Engine/Plugins/Runtime/AudioModulation/Source/AudioModulationEditor/Private/Insights/Providers/ModulationMatrixTraceProvider.h
new file mode 100644
--- /dev/null
+++ b/Engine/Plugins/Runtime/AudioModulation/Source/AudioModulationEditor/Private/Insights/Providers/ModulationMatrixTraceProvider.h
@@ -0,0 +1,52 @@
+// Copyright Epic Games, Inc. All Rights Reserved.
+#pragma once
+
+#include "AudioInsightsTraceChannelHandle.h"
+#include "AudioInsightsTraceProviderBase.h"
+#include "Containers/Map.h"
+#include "Delegates/DelegateCombinations.h"
+#include "Insights/Messages/ModulationMatrixTraceMessages.h"
+
+namespace AudioModulationEditor
+{
+	class FModulationMatrixTraceProvider
+		: public UE::Audio::Insights::TDeviceDataMapTraceProvider<uint32, TSharedPtr<FModulationMatrixDashboardEntry>>
+		, public TSharedFromThis<FModulationMatrixTraceProvider>
+	{
+	public:
+		FModulationMatrixTraceProvider(TSharedRef<UE::Audio::Insights::FTraceChannelManager> InManager);
+
+		virtual ~FModulationMatrixTraceProvider();
+
+		static FName GetName_Static();
+
+		virtual bool ProcessMessages() override;
+		virtual UE::Trace::IAnalyzer* ConstructAnalyzer(TraceServices::IAnalysisSession& InSession) override;
+
+		struct BusInfo
+		{
+			FString BusName;
+			int32 RefCount = 0;
+		};
+
+		using BusIdToBusInfoMap = TMap<FBusId, BusInfo>;
+
+		DECLARE_DELEGATE_OneParam(FOnControlBusesAdded, const BusIdToBusInfoMap& /*AddedControlBuses*/);
+		FOnControlBusesAdded OnControlBusesAdded;
+
+		DECLARE_DELEGATE_OneParam(FOnControlBusesRemoved, const TArray<FName>& /*RemovedControlBusesNames*/);
+		FOnControlBusesRemoved OnControlBusesRemoved;
+
+	private:
+		void UpdateActiveControlBusesToAdd(const TMap<FBusId, FString>& InBusIdToBusNameMap);
+		void UpdateActiveControlBusesToRemove(const TMap<FBusId, float>& InBusIdToValueMap);
+		void OnAudioDeviceDestroyed(::Audio::FDeviceId InDeviceId);
+
+		FModulationMatrixMessages TraceMessages;
+		TSet<UE::Audio::Insights::FTraceChannelHandle> Channels;
+
+		TMap<::Audio::FDeviceId, TSet<uint32>> DeviceIdToActiveModulatorSourceIdsMap;
+		BusIdToBusInfoMap ActiveControlBuses;
+		TArray<FName> RemovedControlBusesNames;
+	};
+} // namespace AudioModulationEditor
