[Audio Modulation Insights] Implemented Modulation Matrix dashboard for Audio Insights (frontend)

diff --git a/Engine/Plugins/Runtime/AudioModulation/Source/AudioModulationEditor/Private/AudioModulationEditor.cpp b/Engine/Plugins/Runtime/AudioModulation/Source/AudioModulationEditor/Private/AudioModulationEditor.cpp
--- a/Engine/Plugins/Runtime/AudioModulation/Source/AudioModulationEditor/Private/AudioModulationEditor.cpp
+++ b/Engine/Plugins/Runtime/AudioModulation/Source/AudioModulationEditor/Private/AudioModulationEditor.cpp
@@ -17,6 +17,7 @@
 #include "IAudioModulation.h"
 #include "ICurveEditorModule.h"
 #include "Insights/Views/ControlBusDashboardViewFactory.h"
+#include "Insights/Views/ModulationMatrixDashboardViewFactory.h"
 #include "Internationalization/Internationalization.h"
 #include "Layouts/SoundControlBusMixStageLayout.h"
 #include "Layouts/SoundControlModulationPatchLayout.h"
@@ -65,6 +66,7 @@ namespace
 			{
 				IAudioInsightsEditorModule& InsightsEditorModule = IAudioInsightsEditorModule::GetChecked();
 				InsightsEditorModule.RegisterDashboardViewFactory(MakeShared<AudioModulationEditor::FControlBusDashboardViewFactory>());
+				InsightsEditorModule.RegisterDashboardViewFactory(MakeShared<AudioModulationEditor::FModulationMatrixDashboardViewFactory>());
 
 				bIsModulationRegisteredInAudioInsights = true;
 			}
@@ -79,6 +81,7 @@ namespace
 			{
 				IAudioInsightsEditorModule& InsightsEditorModule = IAudioInsightsEditorModule::GetChecked();
 				InsightsEditorModule.UnregisterDashboardViewFactory("ControlBuses");
+				InsightsEditorModule.UnregisterDashboardViewFactory("ModulationMatrix");
 
 				bIsModulationRegisteredInAudioInsights = false;
 			}
diff --git a/Engine/Plugins/Runtime/AudioModulation/Source/AudioModulationEditor/Private/Insights/Views/ModulationMatrixDashboardViewFactory.cpp b/Engine/Plugins/Runtime/AudioModulation/Source/AudioModulationEditor/Private/Insights/Views/ModulationMatrixDashboardViewFactory.cpp
new file mode 100644
--- /dev/null
+++ b/Engine/Plugins/Runtime/AudioModulation/Source/AudioModulationEditor/Private/Insights/Views/ModulationMatrixDashboardViewFactory.cpp
@@ -0,0 +1,485 @@
+// Copyright Epic Games, Inc. All Rights Reserved.
+#include "ModulationMatrixDashboardViewFactory.h"
+
+#include "AudioDeviceManager.h"
+#include "AudioInsightsStyle.h"
+#include "AudioInsightsTraceProviderBase.h"
+#include "Framework/Application/SlateApplication.h"
+#include "Fonts/FontMeasure.h"
+#include "IAudioInsightsEditorModule.h"
+#include "Internationalization/Text.h"
+#include "Rendering/SlateRenderer.h"
+#include "Templates/SharedPointer.h"
+#include "Widgets/Input/SComboBox.h"
+
+#define LOCTEXT_NAMESPACE "AudioModulationInsights"
+
+namespace AudioModulationEditor
+{
+	namespace FModulationMatrixDashboardViewFactoryPrivate
+	{
+		const FName SourceIdColumnName("SourceId");
+		const FName ModulatingSourceColumnName("ModulatingSource");
+		const FName ModulatingSourceTypeColumnName("EntryType");
+
+		const FModulationMatrixDashboardEntry& CastEntry(const UE::Audio::Insights::IDashboardDataViewEntry& InData)
+		{
+			return static_cast<const FModulationMatrixDashboardEntry&>(InData);
+		}		
+	}
+
+	FModulationMatrixDashboardViewFactory::FModulationMatrixDashboardViewFactory()
+	{
+		IAudioInsightsTraceModule& AudioInsightsTraceModule = IAudioInsightsEditorModule::GetChecked().GetTraceModule();
+
+		ModulationMatrixTraceProvider = MakeShared<FModulationMatrixTraceProvider>(AudioInsightsTraceModule.GetChannelManager());
+
+		AudioInsightsTraceModule.AddTraceProvider(ModulationMatrixTraceProvider);
+
+		Providers =
+		{
+			ModulationMatrixTraceProvider
+		};
+
+		CreateDefaultColumnData();
+	}
+
+	FName FModulationMatrixDashboardViewFactory::GetName() const
+	{
+		static const FName ModulationMatrixDashboardName = "ModulationMatrix";
+		return ModulationMatrixDashboardName;
+	}
+
+	FText FModulationMatrixDashboardViewFactory::GetDisplayName() const
+	{
+		static const FText ModulationMatrixDashboardDisplayName = LOCTEXT("AudioInsights_ModulationMatrix_DisplayName", "Modulation Matrix");
+		return ModulationMatrixDashboardDisplayName;
+	}
+
+	FSlateIcon FModulationMatrixDashboardViewFactory::GetIcon() const
+	{
+		static const FSlateIcon ModulationMatrixDashboardIcon = { "AudioModulationStyle", "ClassIcon.SoundControlBusMix" };
+		return ModulationMatrixDashboardIcon;
+	}
+
+	UE::Audio::Insights::EDefaultDashboardTabStack FModulationMatrixDashboardViewFactory::GetDefaultTabStack() const
+	{
+		return UE::Audio::Insights::EDefaultDashboardTabStack::Analysis;
+	}
+
+	const TMap<FName, UE::Audio::Insights::FTraceTableDashboardViewFactory::FColumnData>& FModulationMatrixDashboardViewFactory::GetColumns() const
+	{
+		return ModulationMatrixColumnData;
+	}
+
+	void FModulationMatrixDashboardViewFactory::CreateDefaultColumnData()
+	{
+		using namespace UE::Audio::Insights;
+		using namespace FModulationMatrixDashboardViewFactoryPrivate;
+
+		if (!ModulationMatrixColumnData.IsEmpty())
+		{
+			return;
+		}
+
+		ModulationMatrixColumnData =
+		{
+			{
+				SourceIdColumnName,
+				{
+					LOCTEXT("ModulationMatrix_SourceIdColumnDisplayName", "Source Id"),
+					[](const IDashboardDataViewEntry& InData)
+					{
+						const FSourceId SourceId = CastEntry(InData).SourceId;
+						return SourceId != INDEX_NONE ? FText::AsNumber(SourceId) : FText::GetEmpty();
+					},
+					true /* bDefaultHidden */,
+					0.1f /* FillWidth */
+				}
+			},
+			{
+				ModulatingSourceColumnName,
+				{
+					LOCTEXT("ModulationMatrix_ModulatingSourceDisplayName", "Modulating Source"),
+					[](const IDashboardDataViewEntry& InData) { return FText::FromString(CastEntry(InData).Name); },
+					false /* bDefaultHidden */,
+					0.4f /* FillWidth */
+				}
+			},
+			{
+				ModulatingSourceTypeColumnName,
+				{
+					LOCTEXT("ModulationMatrix_ModulatingSourceTypeDisplayName", "Type"),
+					[](const IDashboardDataViewEntry& InData)
+					{
+						const EModulationMatrixEntryType EntryType = CastEntry(InData).EntryType;
+
+						switch (EntryType)
+						{
+							case EModulationMatrixEntryType::BusMix:
+								return LOCTEXT("ModulationMatrix_EntryType", "Bus Mix");
+
+							case EModulationMatrixEntryType::Generator:
+								return LOCTEXT("ModulationMatrix_EntryType", "Generator");
+
+							default:
+								break;
+						}
+
+						return FText::GetEmpty();
+					},
+					false /* bDefaultHidden */,
+					0.125f /* FillWidth */
+				}
+			}
+		};
+
+		SortByColumn = ModulatingSourceColumnName;
+		SortMode = EColumnSortMode::Ascending;
+	}
+
+	void FModulationMatrixDashboardViewFactory::RegisterDelegates()
+	{
+		if (!ModulationMatrixTraceProvider->OnControlBusesAdded.IsBound())
+		{
+			ModulationMatrixTraceProvider->OnControlBusesAdded.BindSP(this, &FModulationMatrixDashboardViewFactory::OnControlBusesAdded);
+		}
+
+		if (!ModulationMatrixTraceProvider->OnControlBusesRemoved.IsBound())
+		{
+			ModulationMatrixTraceProvider->OnControlBusesRemoved.BindSP(this, &FModulationMatrixDashboardViewFactory::OnControlBusesRemoved);
+		}
+
+		if (!OnDeviceDestroyedHandle.IsValid())
+		{
+			OnDeviceDestroyedHandle = FAudioDeviceManagerDelegates::OnAudioDeviceDestroyed.AddSP(this, &FModulationMatrixDashboardViewFactory::OnAudioDeviceDestroyed);
+		}
+	}
+
+	TSharedRef<SWidget> FModulationMatrixDashboardViewFactory::MakeModulatingSourceTypeFilterWidget()
+	{
+		if (ModulatingSourceTypes.IsEmpty())
+		{
+			ModulatingSourceTypes.Emplace(MakeShared<FComboboxSelectionItem>(EModulatingSourceComboboxSelection::All,        LOCTEXT("ModulationMatrix_ModulatingSourceTypeAll",        "All")));
+			ModulatingSourceTypes.Emplace(MakeShared<FComboboxSelectionItem>(EModulatingSourceComboboxSelection::BusMixes,   LOCTEXT("ModulationMatrix_ModulatingSourceTypeBusMixes",   "Bus Mixes")));
+			ModulatingSourceTypes.Emplace(MakeShared<FComboboxSelectionItem>(EModulatingSourceComboboxSelection::Generators, LOCTEXT("ModulationMatrix_ModulatingSourceTypeGenerators", "Generators")));
+
+			SelectedModulatingSourceType = ModulatingSourceTypes[0];
+		}
+
+		return SNew(SHorizontalBox)
+			+ SHorizontalBox::Slot()
+			.AutoWidth()
+			.Padding(2.0f, 10.0f, 0.0f, 0.0f)
+			[
+				SNew(STextBlock)
+				.Margin(FMargin(0.0, 2.0, 0.0, 0.0))
+				.Text(LOCTEXT("ModulationMatrix_TypeFilterText", "Type Filter:"))
+			]
+			+ SHorizontalBox::Slot()
+			.MaxWidth(2.0f)
+			.Padding(0.0f, 10.0f, 0.0f, 0.0f)
+			[
+				SNew(SBox)
+			]
+			+ SHorizontalBox::Slot()
+			.AutoWidth()
+			.HAlign(EHorizontalAlignment::HAlign_Center)
+			.Padding(0.0f, 10.0f, 0.0f, 0.0f)
+			[
+				SNew(SComboBox<TSharedPtr<FComboboxSelectionItem>>)
+				.OptionsSource(&ModulatingSourceTypes)
+				.OnGenerateWidget_Lambda([this](const TSharedPtr<FComboboxSelectionItem>& ModulatingSourceTypePtr)
+				{
+					const FText ModulatingSourceTypeDisplayName = ModulatingSourceTypePtr.IsValid() ? ModulatingSourceTypePtr->Value /*DisplayName*/ : FText::GetEmpty();
+
+					return SNew(STextBlock)
+						.Text(ModulatingSourceTypeDisplayName);
+				})
+				.OnSelectionChanged_Lambda([this](TSharedPtr<FComboboxSelectionItem> InSelectedModulatingSourceTypePtr, ESelectInfo::Type)
+				{
+					if (InSelectedModulatingSourceTypePtr.IsValid())
+					{
+						SelectedModulatingSourceType = InSelectedModulatingSourceTypePtr;
+						UpdateFilterReason = EProcessReason::FilterUpdated;
+					}
+				})
+				[
+					SNew(STextBlock)
+					.Text_Lambda([this]()
+					{
+						const int32 FoundIndex = ModulatingSourceTypes.Find(SelectedModulatingSourceType);
+						if (ModulatingSourceTypes.IsValidIndex(FoundIndex) && ModulatingSourceTypes[FoundIndex].IsValid())
+						{
+							return ModulatingSourceTypes[FoundIndex]->Value;
+						}
+
+						return FText::GetEmpty();
+					})
+				]
+			];
+	}
+
+	TSharedRef<SWidget> FModulationMatrixDashboardViewFactory::MakeWidget(TSharedRef<SDockTab> OwnerTab)
+	{
+		using namespace UE::Audio::Insights;
+
+		RegisterDelegates();
+
+		return SNew(SVerticalBox)
+		+ SVerticalBox::Slot()
+		.AutoHeight()
+		.HAlign(HAlign_Fill)
+		.Padding(0.0f, 0.0f, 0.0f, 6.0f)
+		[
+			MakeModulatingSourceTypeFilterWidget()
+		]
+		+ SVerticalBox::Slot()
+		.HAlign(HAlign_Fill)
+		[
+			FTraceTableDashboardViewFactory::MakeWidget(OwnerTab)
+		];
+	}
+
+	void FModulationMatrixDashboardViewFactory::ProcessEntries(UE::Audio::Insights::FTraceTableDashboardViewFactory::EProcessReason Reason)
+	{
+		FilterByModulatingSourceName();
+		FilterByModulatingSourceType();
+	}
+
+	void FModulationMatrixDashboardViewFactory::OnControlBusesAdded(const FModulationMatrixTraceProvider::BusIdToBusInfoMap& AddedControlBuses)
+	{
+		using namespace UE::Audio::Insights;
+		using namespace FModulationMatrixDashboardViewFactoryPrivate;
+
+		for (const auto& [BusId, BusInfo] : AddedControlBuses)
+		{
+			const FName BusFName(BusInfo.BusName);
+
+			if (HeaderRowWidget->IsColumnGenerated(BusFName))
+			{
+				continue;
+			}
+
+			const FString& BusName = BusInfo.BusName;
+			const FText BusNameText = FText::FromString(BusName);
+
+			// Add bus infomation to column data
+			ModulationMatrixColumnData.Add(BusFName,
+			{
+				BusNameText,
+				[&BusId, &BusName](const IDashboardDataViewEntry& InData)
+				{
+					const FModulationMatrixDashboardEntry& ModulationMatrixEntry = CastEntry(InData);
+
+					if (const float* FoundValuePtr = ModulationMatrixEntry.BusIdToValueMap.Find(BusId))
+					{
+						return FText::AsNumber(*FoundValuePtr, FSlateStyle::Get().GetAmpFloatFormat());
+					}
+
+					return FText::GetEmpty();
+				},
+				false /* bDefaultHidden */,
+				0.1f  /* FillWidth */
+			});
+
+			// Add column to widget
+			SHeaderRow::FColumn::FArguments ColumnArgs = SHeaderRow::Column(BusFName)
+				.DefaultLabel(BusNameText)
+				.HAlignCell(HAlign_Left);
+
+			if (const FSlateRenderer* SlateRenderer = FSlateApplication::Get().GetRenderer())
+			{
+				const FVector2D TextSize = SlateRenderer->GetFontMeasureService()->Measure(BusNameText, FAppStyle::GetFontStyle("NormalFont"));
+
+				ColumnArgs.ManualWidth(TextSize.X + 10.0);
+			}
+
+			HeaderRowWidget->AddColumn(ColumnArgs);
+
+			ActiveBusNames.Emplace(BusFName);
+		}
+	}
+
+	void FModulationMatrixDashboardViewFactory::OnControlBusesRemoved(const TArray<FName>& RemovedControlBusesNames)
+	{
+		for (const FName& BusName : RemovedControlBusesNames)
+		{
+			HeaderRowWidget->RemoveColumn(BusName);
+			ModulationMatrixColumnData.Remove(BusName);
+			ActiveBusNames.Remove(BusName);
+		}
+	}
+
+	void FModulationMatrixDashboardViewFactory::OnAudioDeviceDestroyed(::Audio::FDeviceId InDeviceId)
+	{
+		for (const FName& BusName : ActiveBusNames)
+		{
+			HeaderRowWidget->RemoveColumn(BusName);
+			ModulationMatrixColumnData.Remove(BusName);
+		}
+
+		ActiveBusNames.Empty();
+	}
+
+	void FModulationMatrixDashboardViewFactory::FilterByModulatingSourceName()
+	{
+		using namespace UE::Audio::Insights;
+		
+		const FString FilterString = GetSearchFilterText().ToString();
+
+		FTraceTableDashboardViewFactory::FilterEntries<FModulationMatrixTraceProvider>(
+		[&FilterString](const IDashboardDataViewEntry& Entry)
+		{
+			const FModulationMatrixDashboardEntry& ModulationMatrixEntry = FModulationMatrixDashboardViewFactoryPrivate::CastEntry(Entry);
+
+			return !ModulationMatrixEntry.Name.Contains(FilterString) && ModulationMatrixEntry.EntryType != EModulationMatrixEntryType::BusFinalValues;
+		});
+	}
+
+	void FModulationMatrixDashboardViewFactory::FilterByModulatingSourceType()
+	{
+		using namespace UE::Audio::Insights;
+
+		TArray<TSharedPtr<IDashboardDataViewEntry>> EntriesToFilterOut;
+
+		const EModulatingSourceComboboxSelection SelectedModulatingSourceTypeEnum = SelectedModulatingSourceType.IsValid()? SelectedModulatingSourceType->Key : EModulatingSourceComboboxSelection::All;
+
+		for (const TSharedPtr<IDashboardDataViewEntry>& Entry : DataViewEntries)
+		{
+			if (Entry.IsValid())
+			{
+				const FModulationMatrixDashboardEntry& ModulationMatrixEntry = FModulationMatrixDashboardViewFactoryPrivate::CastEntry(*Entry);
+
+				// We never want to filter out the bus final values row
+				if (ModulationMatrixEntry.EntryType == EModulationMatrixEntryType::BusFinalValues)
+				{
+					continue;
+				}
+
+				if (SelectedModulatingSourceTypeEnum != EModulatingSourceComboboxSelection::All &&
+					(SelectedModulatingSourceTypeEnum == EModulatingSourceComboboxSelection::BusMixes   && ModulationMatrixEntry.EntryType != EModulationMatrixEntryType::BusMix) ||
+					(SelectedModulatingSourceTypeEnum == EModulatingSourceComboboxSelection::Generators && ModulationMatrixEntry.EntryType != EModulationMatrixEntryType::Generator))
+				{
+					EntriesToFilterOut.Emplace(Entry);
+				}
+			}
+		}
+
+		for (const TSharedPtr<IDashboardDataViewEntry>& Entry : EntriesToFilterOut)
+		{
+			DataViewEntries.Remove(Entry);
+		}
+	}
+
+	void FModulationMatrixDashboardViewFactory::SortTable()
+	{
+		using namespace UE::Audio::Insights;
+		using namespace FModulationMatrixDashboardViewFactoryPrivate;
+
+		if (SortByColumn == SourceIdColumnName)
+		{
+			if (SortMode == EColumnSortMode::Ascending)
+			{
+				DataViewEntries.Sort([](const TSharedPtr<IDashboardDataViewEntry>& A, const TSharedPtr<IDashboardDataViewEntry>& B)
+				{
+					const FModulationMatrixDashboardEntry& AData = CastEntry(*A);
+					const FModulationMatrixDashboardEntry& BData = CastEntry(*B);
+
+					if (AData.EntryType == EModulationMatrixEntryType::BusFinalValues) { return false; }
+					if (BData.EntryType == EModulationMatrixEntryType::BusFinalValues) { return true;  }
+
+					return AData.SourceId < BData.SourceId;
+				});
+			}
+			else if (SortMode == EColumnSortMode::Descending)
+			{
+				DataViewEntries.Sort([](const TSharedPtr<IDashboardDataViewEntry>& A, const TSharedPtr<IDashboardDataViewEntry>& B)
+				{
+					const FModulationMatrixDashboardEntry& AData = CastEntry(*A);
+					const FModulationMatrixDashboardEntry& BData = CastEntry(*B);
+
+					if (AData.EntryType == EModulationMatrixEntryType::BusFinalValues) { return false; }
+					if (BData.EntryType == EModulationMatrixEntryType::BusFinalValues) { return true;  }
+
+					return BData.SourceId < AData.SourceId;
+				});
+			}
+		}
+		else if (SortByColumn == ModulatingSourceColumnName)
+		{
+			if (SortMode == EColumnSortMode::Ascending)
+			{
+				DataViewEntries.Sort([](const TSharedPtr<IDashboardDataViewEntry>& A, const TSharedPtr<IDashboardDataViewEntry>& B)
+				{
+					const FModulationMatrixDashboardEntry& AData = CastEntry(*A);
+					const FModulationMatrixDashboardEntry& BData = CastEntry(*B);
+
+					if (AData.EntryType == EModulationMatrixEntryType::BusFinalValues) { return false; }
+					if (BData.EntryType == EModulationMatrixEntryType::BusFinalValues) { return true;  }
+
+					return AData.Name.ToLower() < BData.Name.ToLower();
+				});
+			}
+			else if (SortMode == EColumnSortMode::Descending)
+			{
+				DataViewEntries.Sort([](const TSharedPtr<IDashboardDataViewEntry>& A, const TSharedPtr<IDashboardDataViewEntry>& B)
+				{
+					const FModulationMatrixDashboardEntry& AData = CastEntry(*A);
+					const FModulationMatrixDashboardEntry& BData = CastEntry(*B);
+
+					if (AData.EntryType == EModulationMatrixEntryType::BusFinalValues) { return false; }
+					if (BData.EntryType == EModulationMatrixEntryType::BusFinalValues) { return true;  }
+
+					return BData.Name.ToLower() < AData.Name.ToLower();
+				});
+			}
+		}
+		else if (SortByColumn == ModulatingSourceTypeColumnName)
+		{
+			if (SortMode == EColumnSortMode::Ascending)
+			{
+				DataViewEntries.Sort([](const TSharedPtr<IDashboardDataViewEntry>& A, const TSharedPtr<IDashboardDataViewEntry>& B)
+				{
+					const FModulationMatrixDashboardEntry& AData = CastEntry(*A);
+					const FModulationMatrixDashboardEntry& BData = CastEntry(*B);
+
+					if (AData.EntryType == EModulationMatrixEntryType::BusFinalValues) { return false; }
+					if (BData.EntryType == EModulationMatrixEntryType::BusFinalValues) { return true;  }
+
+					return AData.EntryType < BData.EntryType;
+				});
+			}
+			else if (SortMode == EColumnSortMode::Descending)
+			{
+				DataViewEntries.Sort([](const TSharedPtr<IDashboardDataViewEntry>& A, const TSharedPtr<IDashboardDataViewEntry>& B)
+				{
+					const FModulationMatrixDashboardEntry& AData = CastEntry(*A);
+					const FModulationMatrixDashboardEntry& BData = CastEntry(*B);
+
+					if (AData.EntryType == EModulationMatrixEntryType::BusFinalValues) { return false; }
+					if (BData.EntryType == EModulationMatrixEntryType::BusFinalValues) { return true;  }
+
+					return BData.EntryType < AData.EntryType;
+				});
+			}
+		}
+	}
+
+	FSlateColor FModulationMatrixDashboardViewFactory::GetRowColor(const TSharedPtr<UE::Audio::Insights::IDashboardDataViewEntry>& InRowDataPtr)
+	{
+		if (InRowDataPtr.IsValid())
+		{
+			const FModulationMatrixDashboardEntry& ModulationMatrixEntry = FModulationMatrixDashboardViewFactoryPrivate::CastEntry(*InRowDataPtr);
+
+			if (ModulationMatrixEntry.EntryType == EModulationMatrixEntryType::BusFinalValues)
+			{
+				return FSlateColor(FColor::Green);
+			}
+		}
+
+		return FSlateColor(FColor::White);
+	}
+} // namespace AudioModulationEditor
+
+#undef LOCTEXT_NAMESPACE
diff --git a/Engine/Plugins/Runtime/AudioModulation/Source/AudioModulationEditor/Private/Insights/Views/ModulationMatrixDashboardViewFactory.h b/Engine/Plugins/Runtime/AudioModulation/Source/AudioModulationEditor/Private/Insights/Views/ModulationMatrixDashboardViewFactory.h
new file mode 100644
--- /dev/null
+++ b/Engine/Plugins/Runtime/AudioModulation/Source/AudioModulationEditor/Private/Insights/Views/ModulationMatrixDashboardViewFactory.h
@@ -0,0 +1,68 @@
+// Copyright Epic Games, Inc. All Rights Reserved.
+#pragma once
+
+#include "Containers/Map.h"
+#include "Insights/Providers/ModulationMatrixTraceProvider.h"
+#include "Templates/SharedPointer.h"
+#include "Views/TableDashboardViewFactory.h"
+
+namespace UE::Audio::Insights
+{
+	class IDashboardDataViewEntry;
+}
+
+namespace AudioModulationEditor
+{
+	class FModulationMatrixDashboardViewFactory : public UE::Audio::Insights::FTraceObjectTableDashboardViewFactory
+	{
+	public:
+		FModulationMatrixDashboardViewFactory();
+		virtual ~FModulationMatrixDashboardViewFactory() = default;
+
+		virtual FName GetName() const override;
+		virtual FText GetDisplayName() const override;
+		virtual FSlateIcon GetIcon() const override;
+		virtual UE::Audio::Insights::EDefaultDashboardTabStack GetDefaultTabStack() const override;
+
+		virtual TSharedRef<SWidget> MakeWidget(TSharedRef<SDockTab> OwnerTab) override;
+
+		virtual void ProcessEntries(UE::Audio::Insights::FTraceTableDashboardViewFactory::EProcessReason Reason) override;
+
+	private:
+		virtual const TMap<FName, UE::Audio::Insights::FTraceTableDashboardViewFactory::FColumnData>& GetColumns() const override;
+
+		void CreateDefaultColumnData();
+		void RegisterDelegates();
+
+		TSharedRef<SWidget> MakeModulatingSourceTypeFilterWidget();
+
+		void OnControlBusesAdded(const FModulationMatrixTraceProvider::BusIdToBusInfoMap& AddedControlBuses);
+		void OnControlBusesRemoved(const TArray<FName>& RemovedControlBusesNames);
+		void OnAudioDeviceDestroyed(::Audio::FDeviceId InDeviceId);
+
+		void FilterByModulatingSourceName();
+		void FilterByModulatingSourceType();
+
+		virtual void SortTable() override;
+		virtual FSlateColor GetRowColor(const TSharedPtr<UE::Audio::Insights::IDashboardDataViewEntry>& InRowDataPtr) override;
+
+		enum class EModulatingSourceComboboxSelection : uint8
+		{
+			All,
+			BusMixes,
+			Generators
+		};
+
+		TSharedPtr<FModulationMatrixTraceProvider> ModulationMatrixTraceProvider;
+		TMap<FName, FTraceTableDashboardViewFactory::FColumnData> ModulationMatrixColumnData;
+		TArray<FName> ActiveBusNames;
+
+		FDelegateHandle OnControlBusAddedHandle;
+		FDelegateHandle OnControlBusRemovedHandle;
+		FDelegateHandle OnDeviceDestroyedHandle;
+
+		using FComboboxSelectionItem = TPair<EModulatingSourceComboboxSelection, FText>;
+		TArray<TSharedPtr<FComboboxSelectionItem>> ModulatingSourceTypes;
+		TSharedPtr<FComboboxSelectionItem> SelectedModulatingSourceType;
+	};
+} // namespace AudioModulationEditor
