[Audio Modulation Insights] Static code analyzer fix for: Implemented Modulation Matrix dashboard for Audio Insights (backend)

diff --git a/Engine/Plugins/Runtime/AudioModulation/Source/AudioModulationEditor/Private/Insights/Views/ModulationMatrixDashboardViewFactory.cpp b/Engine/Plugins/Runtime/AudioModulation/Source/AudioModulationEditor/Private/Insights/Views/ModulationMatrixDashboardViewFactory.cpp
--- a/Engine/Plugins/Runtime/AudioModulation/Source/AudioModulationEditor/Private/Insights/Views/ModulationMatrixDashboardViewFactory.cpp
+++ b/Engine/Plugins/Runtime/AudioModulation/Source/AudioModulationEditor/Private/Insights/Views/ModulationMatrixDashboardViewFactory.cpp
@@ -357,11 +357,13 @@ namespace AudioModulationEditor
 					continue;
 				}
 
-				if (SelectedModulatingSourceTypeEnum != EModulatingSourceComboboxSelection::All &&
-					(SelectedModulatingSourceTypeEnum == EModulatingSourceComboboxSelection::BusMixes   && ModulationMatrixEntry.EntryType != EModulationMatrixEntryType::BusMix) ||
-					(SelectedModulatingSourceTypeEnum == EModulatingSourceComboboxSelection::Generators && ModulationMatrixEntry.EntryType != EModulationMatrixEntryType::Generator))
+				if (SelectedModulatingSourceTypeEnum != EModulatingSourceComboboxSelection::All)
 				{
-					EntriesToFilterOut.Emplace(Entry);
+					if ((SelectedModulatingSourceTypeEnum == EModulatingSourceComboboxSelection::BusMixes   && ModulationMatrixEntry.EntryType != EModulationMatrixEntryType::BusMix) ||
+						(SelectedModulatingSourceTypeEnum == EModulatingSourceComboboxSelection::Generators && ModulationMatrixEntry.EntryType != EModulationMatrixEntryType::Generator))
+					{
+						EntriesToFilterOut.Emplace(Entry);
+					}
 				}
 			}
 		}
