[Audio Modulation Insights] Fixed build issue related to LOCTEXT.

diff --git a/Engine/Plugins/Runtime/AudioModulation/Source/AudioModulationEditor/Private/Insights/Views/ModulationMatrixDashboardViewFactory.cpp b/Engine/Plugins/Runtime/AudioModulation/Source/AudioModulationEditor/Private/Insights/Views/ModulationMatrixDashboardViewFactory.cpp
--- a/Engine/Plugins/Runtime/AudioModulation/Source/AudioModulationEditor/Private/Insights/Views/ModulationMatrixDashboardViewFactory.cpp
+++ b/Engine/Plugins/Runtime/AudioModulation/Source/AudioModulationEditor/Private/Insights/Views/ModulationMatrixDashboardViewFactory.cpp
@@ -117,10 +117,10 @@ namespace AudioModulationEditor
 						switch (EntryType)
 						{
 							case EModulationMatrixEntryType::BusMix:
-								return LOCTEXT("ModulationMatrix_EntryType", "Bus Mix");
+								return LOCTEXT("ModulationMatrix_EntryTypeBusMix", "Bus Mix");
 
 							case EModulationMatrixEntryType::Generator:
-								return LOCTEXT("ModulationMatrix_EntryType", "Generator");
+								return LOCTEXT("ModulationMatrix_EntryTypeGenerator", "Generator");
 
 							default:
 								break;
