ModulationPatchEditorViewStacked::FormatTooltipValue

- Added override of FormatToolTipValue for the ModulationPatchEditor so that when you hover over a point on a curve, it shows you the unit value accompanied by its unit as well as the normalised value

[RN] minor framework/audio

diff --git a/Engine/Plugins/Runtime/AudioModulation/Source/AudioModulationEditor/Private/Editors/ModulationPatchCurveEditorViewStacked.cpp b/Engine/Plugins/Runtime/AudioModulation/Source/AudioModulationEditor/Private/Editors/ModulationPatchCurveEditorViewStacked.cpp
--- a/Engine/Plugins/Runtime/AudioModulation/Source/AudioModulationEditor/Private/Editors/ModulationPatchCurveEditorViewStacked.cpp
+++ b/Engine/Plugins/Runtime/AudioModulation/Source/AudioModulationEditor/Private/Editors/ModulationPatchCurveEditorViewStacked.cpp
@@ -2,15 +2,11 @@
 
 #include "ModulationPatchCurveEditorViewStacked.h"
 
-#include "AudioModulationStyle.h"
 #include "Curves/CurveFloat.h"
 #include "CurveEditor.h"
-#include "CurveModel.h"
 #include "Styling/AppStyle.h"
-#include "Fonts/FontMeasure.h"
 #include "Framework/Application/SlateApplication.h"
 #include "IAudioModulation.h"
-#include "SCurveEditorPanel.h"
 #include "Sound/SoundModulationDestination.h"
 #include "SoundControlBus.h"
 #include "SoundModulationParameter.h"
@@ -60,6 +56,29 @@ void SModulationPatchEditorViewStacked::FormatOutputLabel(const WaveTable::Edito
 	}
 }
 
+FText SModulationPatchEditorViewStacked::FormatToolTipValue(const FCurveModel& CurveModel, double EvaluatedValue) const
+{
+	FNumberFormattingOptions FormatOptions;
+	FormatOptions.MaximumFractionalDigits = MaximumFractionalDigits;
+	const FText NormalizedValue = FText::AsNumber(EvaluatedValue, &FormatOptions);
+	const FModPatchCurveEditorModel& PatchCurveModel = static_cast<const FModPatchCurveEditorModel&>(CurveModel);
+	const USoundModulationPatch* Patch = PatchCurveModel.GetPatch();
+	FText UnitValue;
+	FText UnitLabel;
+	if (const TObjectPtr<USoundModulationParameter> Parameter = Patch->PatchSettings.OutputParameter)
+	{
+		const float Value = Parameter->ConvertNormalizedToUnit(EvaluatedValue);
+		// Set the fractional digits based on the difference between the minimum and maximum value of the parameter.
+		// Only add fractional digits so that the total number of digits is at least 4
+		const float MinMaxDiff = Parameter->GetUnitMax() - Parameter->GetUnitMin();
+		const int DiffDigits = FMath::Floor(FMath::LogX(10, FMath::Abs(MinMaxDiff)));
+		FormatOptions.MaximumFractionalDigits = FMath::Clamp(MaximumFractionalDigits - DiffDigits, 0, MaximumFractionalDigits);
+		UnitValue = FText::AsNumber(Value, &FormatOptions);
+		UnitLabel = Parameter->Settings.UnitDisplayName;
+	}
+	return FText::Format(LOCTEXT("ModulationPatch_PointToolTipValue", "Value:  {0}{1} ({2})"), UnitValue, UnitLabel, NormalizedValue);
+}
+
 ECurveEditorViewID FModPatchCurveEditorModel::ModPatchViewId = ECurveEditorViewID::Invalid;
 
 FModPatchCurveEditorModel::FModPatchCurveEditorModel(FRichCurve& InRichCurve, UObject* InOwner, EWaveTableCurveSource InSource)
diff --git a/Engine/Plugins/Runtime/AudioModulation/Source/AudioModulationEditor/Private/Editors/ModulationPatchCurveEditorViewStacked.h b/Engine/Plugins/Runtime/AudioModulation/Source/AudioModulationEditor/Private/Editors/ModulationPatchCurveEditorViewStacked.h
--- a/Engine/Plugins/Runtime/AudioModulation/Source/AudioModulationEditor/Private/Editors/ModulationPatchCurveEditorViewStacked.h
+++ b/Engine/Plugins/Runtime/AudioModulation/Source/AudioModulationEditor/Private/Editors/ModulationPatchCurveEditorViewStacked.h
@@ -2,14 +2,11 @@
 #pragma once
 
 #include "AudioModulationStyle.h"
-#include "CurveModel.h"
 #include "Curves/RichCurve.h"
 #include "RichCurveEditorModel.h"
 #include "UObject/Object.h"
 #include "UObject/WeakObjectPtr.h"
 #include "UObject/WeakObjectPtrTemplates.h"
-#include "Views/SCurveEditorViewStacked.h"
-#include "Views/SInteractiveCurveEditorView.h"
 #include "WaveTableCurveEditorViewStacked.h"
 
 
@@ -25,6 +22,11 @@ class SModulationPatchEditorViewStacked : public WaveTable::Editor::SViewStacked
 {
 	virtual void FormatInputLabel(const WaveTable::Editor::FWaveTableCurveModel& EditorModel, const FNumberFormattingOptions& InLabelFormat, FText& InOutLabel) const override;
 	virtual void FormatOutputLabel(const WaveTable::Editor::FWaveTableCurveModel& EditorModel, const FNumberFormattingOptions& InLabelFormat, FText& InOutLabel) const override;
+
+protected:
+	virtual FText FormatToolTipValue(const FCurveModel& CurveModel, double EvaluatedValue) const override;
+	
+	static constexpr int MaximumFractionalDigits = 3;
 };
 
 class FModPatchCurveEditorModel : public WaveTable::Editor::FWaveTableCurveModel
