[Audio Modulation Insights] ControlBusDashboardViewFactory: Fixes for text filtering and column sorting.

diff --git a/Engine/Plugins/Runtime/AudioModulation/Source/AudioModulationEditor/Private/Insights/Views/ControlBusDashboardViewFactory.cpp b/Engine/Plugins/Runtime/AudioModulation/Source/AudioModulationEditor/Private/Insights/Views/ControlBusDashboardViewFactory.cpp
--- a/Engine/Plugins/Runtime/AudioModulation/Source/AudioModulationEditor/Private/Insights/Views/ControlBusDashboardViewFactory.cpp
+++ b/Engine/Plugins/Runtime/AudioModulation/Source/AudioModulationEditor/Private/Insights/Views/ControlBusDashboardViewFactory.cpp
@@ -64,7 +64,7 @@ namespace AudioModulationEditor
 		FTraceTableDashboardViewFactory::FilterEntries<FControlBusTraceProvider>([&FilterString](const IDashboardDataViewEntry& Entry)
 		{
 			const FControlBusDashboardEntry& ControlBusEntry = ControlBusPrivate::CastEntry(Entry);
-			if (ControlBusEntry.GetDisplayName().ToString().Contains(FilterString))
+			if (ControlBusEntry.Name.Contains(FilterString))
 			{
 				return false;
 			}
@@ -163,7 +163,7 @@ namespace AudioModulationEditor
 	{
 		using namespace UE::Audio::Insights;
 
-		if (SortByColumn == "Name")
+		if (SortByColumn == "BusId")
 		{
 			if (SortMode == EColumnSortMode::Ascending)
 			{
@@ -172,7 +172,7 @@ namespace AudioModulationEditor
 					const FControlBusDashboardEntry& AData = ControlBusPrivate::CastEntry(*A.Get());
 					const FControlBusDashboardEntry& BData = ControlBusPrivate::CastEntry(*B.Get());
 
-					return AData.GetDisplayName().CompareToCaseIgnored(BData.GetDisplayName()) < 0;
+					return AData.ControlBusId < BData.ControlBusId;
 				});
 			}
 			else if (SortMode == EColumnSortMode::Descending)
@@ -182,7 +182,30 @@ namespace AudioModulationEditor
 					const FControlBusDashboardEntry& AData = ControlBusPrivate::CastEntry(*A.Get());
 					const FControlBusDashboardEntry& BData = ControlBusPrivate::CastEntry(*B.Get());
 
-					return BData.GetDisplayName().CompareToCaseIgnored(AData.GetDisplayName()) < 0;
+					return BData.ControlBusId < AData.ControlBusId;
+				});
+			}
+		}
+		else if (SortByColumn == "Name")
+		{
+			if (SortMode == EColumnSortMode::Ascending)
+			{
+				DataViewEntries.Sort([](const TSharedPtr<IDashboardDataViewEntry>& A, const TSharedPtr<IDashboardDataViewEntry>& B)
+				{
+					const FControlBusDashboardEntry& AData = ControlBusPrivate::CastEntry(*A.Get());
+					const FControlBusDashboardEntry& BData = ControlBusPrivate::CastEntry(*B.Get());
+
+					return AData.Name.ToLower() < BData.Name.ToLower();
+				});
+			}
+			else if (SortMode == EColumnSortMode::Descending)
+			{
+				DataViewEntries.Sort([](const TSharedPtr<IDashboardDataViewEntry>& A, const TSharedPtr<IDashboardDataViewEntry>& B)
+				{
+					const FControlBusDashboardEntry& AData = ControlBusPrivate::CastEntry(*A.Get());
+					const FControlBusDashboardEntry& BData = ControlBusPrivate::CastEntry(*B.Get());
+
+					return BData.Name.ToLower() < AData.Name.ToLower();
 				});
 			}
 		}
@@ -195,7 +218,30 @@ namespace AudioModulationEditor
 					const FControlBusDashboardEntry& AData = ControlBusPrivate::CastEntry(*A.Get());
 					const FControlBusDashboardEntry& BData = ControlBusPrivate::CastEntry(*B.Get());
 
-					return AData.GetParamNameAsFText().CompareToCaseIgnored(BData.GetParamNameAsFText()) < 0;
+					return AData.ParamName.ToLower() < BData.ParamName.ToLower();
+				});
+			}
+			else if (SortMode == EColumnSortMode::Descending)
+			{
+				DataViewEntries.Sort([](const TSharedPtr<IDashboardDataViewEntry>& A, const TSharedPtr<IDashboardDataViewEntry>& B)
+				{
+					const FControlBusDashboardEntry& AData = ControlBusPrivate::CastEntry(*A.Get());
+					const FControlBusDashboardEntry& BData = ControlBusPrivate::CastEntry(*B.Get());
+
+					return BData.ParamName.ToLower() < AData.ParamName.ToLower();
+				});
+			}
+		}
+		else if (SortByColumn == "Value")
+		{
+			if (SortMode == EColumnSortMode::Ascending)
+			{
+				DataViewEntries.Sort([](const TSharedPtr<IDashboardDataViewEntry>& A, const TSharedPtr<IDashboardDataViewEntry>& B)
+				{
+					const FControlBusDashboardEntry& AData = ControlBusPrivate::CastEntry(*A.Get());
+					const FControlBusDashboardEntry& BData = ControlBusPrivate::CastEntry(*B.Get());
+
+					return AData.Value < BData.Value;
 				});
 			}
 			else if (SortMode == EColumnSortMode::Descending)
@@ -205,7 +251,7 @@ namespace AudioModulationEditor
 					const FControlBusDashboardEntry& AData = ControlBusPrivate::CastEntry(*A.Get());
 					const FControlBusDashboardEntry& BData = ControlBusPrivate::CastEntry(*B.Get());
 
-					return BData.GetParamNameAsFText().CompareToCaseIgnored(AData.GetParamNameAsFText()) < 0;
+					return BData.Value < AData.Value;
 				});
 			}
 		}
