[Audio Modulation Insights] AudioModulationEditor.cpp: removed au.Modulation.AudioInsightsTab cvar, from now on the audio modulation insights dashboards will load/unload automatically.

diff --git a/Engine/Plugins/Runtime/AudioModulation/Source/AudioModulationEditor/Private/AudioModulationEditor.cpp b/Engine/Plugins/Runtime/AudioModulation/Source/AudioModulationEditor/Private/AudioModulationEditor.cpp
--- a/Engine/Plugins/Runtime/AudioModulation/Source/AudioModulationEditor/Private/AudioModulationEditor.cpp
+++ b/Engine/Plugins/Runtime/AudioModulation/Source/AudioModulationEditor/Private/AudioModulationEditor.cpp
@@ -87,30 +87,8 @@ namespace
 			}
 		}
 	}
-
-	void OnModulationInsightsCVarValueChanged(int32 CVarValue)
-	{
-		if (CVarValue)
-		{
-			RegisterModulationInAudioInsights();
-		}
-		else
-		{
-			UnregisterModulationInAudioInsights();
-		}
-	}
-
 }
 
-static int32 AudioModulationInsightsTabCVar = 0;
-static FAutoConsoleVariableRef AudioModulationEnableInsightsTab(
-	TEXT("au.Modulation.AudioInsightsTab"),
-	AudioModulationInsightsTabCVar,
-	TEXT("Enable/disable the Audio Modulation Tab in Audio Insights. 1 = enable, 0 = disable\n")
-	TEXT("Requires closing and Reopening Insights to take effect.\n"),
-	FConsoleVariableDelegate::CreateLambda([](IConsoleVariable* Var) { ::OnModulationInsightsCVarValueChanged(AudioModulationInsightsTabCVar); }),
-	ECVF_Default);
-
 FAudioModulationEditorModule::FAudioModulationEditorModule()
 {
 }
@@ -213,11 +191,7 @@ void FAudioModulationEditorModule::StartupModule()
 		}
 	});
 
-
-	if (AudioModulationInsightsTabCVar)
-	{
-		::RegisterModulationInAudioInsights();
-	}
+	::RegisterModulationInAudioInsights();
 }
 
 void FAudioModulationEditorModule::RegisterCustomPropertyLayouts()
@@ -261,10 +235,7 @@ void FAudioModulationEditorModule::ShutdownModule()
 
 	FSlateStyleRegistry::UnRegisterSlateStyle(*StyleSet.Get());
 
-	if (AudioModulationInsightsTabCVar)
-	{
-		::UnregisterModulationInAudioInsights();
-	}
+	::UnregisterModulationInAudioInsights();
 }
 
 IMPLEMENT_MODULE(FAudioModulationEditorModule, AudioModulationEditor);
