Fix All Same-Name crashes in Audio Modulation Static BP Functions

diff --git a/Engine/Plugins/Runtime/AudioModulation/Source/AudioModulation/Private/AudioModulationStatics.cpp b/Engine/Plugins/Runtime/AudioModulation/Source/AudioModulation/Private/AudioModulationStatics.cpp
--- a/Engine/Plugins/Runtime/AudioModulation/Source/AudioModulation/Private/AudioModulationStatics.cpp
+++ b/Engine/Plugins/Runtime/AudioModulation/Source/AudioModulation/Private/AudioModulationStatics.cpp
@@ -204,7 +204,8 @@ USoundControlBus* UAudioModulationStatics::CreateBus(UObject* WorldContextObject
 		return nullptr;
 	}
 
-	USoundControlBus* NewBus = NewObject<USoundControlBus>(WorldContextObject, Name);
+	FName UniqueName = MakeUniqueObjectName(WorldContextObject, USoundControlBus::StaticClass(), Name);
+	USoundControlBus* NewBus = NewObject<USoundControlBus>(WorldContextObject, UniqueName);
 	NewBus->Parameter = Parameter;
 	NewBus->Address = Name.ToString();
 
@@ -237,7 +238,8 @@ USoundControlBusMix* UAudioModulationStatics::CreateBusMix(UObject* WorldContext
 		return nullptr;
 	}
 
-	USoundControlBusMix* NewBusMix = NewObject<USoundControlBusMix>(WorldContextObject, Name);
+	FName UniqueName = MakeUniqueObjectName(WorldContextObject, USoundControlBusMix::StaticClass(), Name);
+	USoundControlBusMix* NewBusMix = NewObject<USoundControlBusMix>(WorldContextObject, UniqueName);
 	for (FSoundControlBusMixStage& Stage : Stages)
 	{
 		if (Stage.Bus)
@@ -271,7 +273,8 @@ USoundModulationParameter* UAudioModulationStatics::CreateModulationParameter(UO
 		return nullptr;
 	}
 
-	USoundModulationParameter* NewParam = NewObject<USoundModulationParameter>(WorldContextObject, ParamClass, Name);
+	FName UniqueName = MakeUniqueObjectName(WorldContextObject, USoundModulationParameter::StaticClass(), Name);
+	USoundModulationParameter* NewParam = NewObject<USoundModulationParameter>(WorldContextObject, ParamClass, UniqueName);
 
 	if (NewParam)
 	{
@@ -288,8 +291,9 @@ USoundModulationGeneratorEnvelopeFollower* UAudioModulationStatics::CreateEnvelo
 	{
 		return nullptr;
 	}
-	
-	USoundModulationGeneratorEnvelopeFollower* NewGenerator = NewObject<USoundModulationGeneratorEnvelopeFollower>(WorldContextObject, USoundModulationGeneratorEnvelopeFollower::StaticClass(), Name);
+
+	FName UniqueName = MakeUniqueObjectName(WorldContextObject, USoundModulationGeneratorEnvelopeFollower::StaticClass(), Name);
+	USoundModulationGeneratorEnvelopeFollower* NewGenerator = NewObject<USoundModulationGeneratorEnvelopeFollower>(WorldContextObject, USoundModulationGeneratorEnvelopeFollower::StaticClass(), UniqueName);
 
 	if (NewGenerator)
 	{
@@ -307,7 +311,8 @@ USoundModulationGeneratorLFO* UAudioModulationStatics::CreateLFOGenerator(UObjec
 		return nullptr;
 	}
 
-	USoundModulationGeneratorLFO* NewGenerator = NewObject<USoundModulationGeneratorLFO>(WorldContextObject, USoundModulationGeneratorLFO::StaticClass(), Name);
+	FName UniqueName = MakeUniqueObjectName(WorldContextObject, USoundModulationGeneratorLFO::StaticClass(), Name);
+	USoundModulationGeneratorLFO* NewGenerator = NewObject<USoundModulationGeneratorLFO>(WorldContextObject, USoundModulationGeneratorLFO::StaticClass(), UniqueName);
 
 	if (NewGenerator)
 	{
@@ -325,7 +330,8 @@ USoundModulationGeneratorADEnvelope* UAudioModulationStatics::CreateADEnvelopeGe
 		return nullptr;
 	}
 
-	USoundModulationGeneratorADEnvelope* NewGenerator = NewObject<USoundModulationGeneratorADEnvelope>(WorldContextObject, USoundModulationGeneratorADEnvelope::StaticClass(), Name);
+	FName UniqueName = MakeUniqueObjectName(WorldContextObject, USoundModulationGeneratorADEnvelope::StaticClass(), Name);
+	USoundModulationGeneratorADEnvelope* NewGenerator = NewObject<USoundModulationGeneratorADEnvelope>(WorldContextObject, USoundModulationGeneratorADEnvelope::StaticClass(), UniqueName);
 
 	if (NewGenerator)
 	{
@@ -343,7 +349,9 @@ UAudioModulationDestination* UAudioModulationStatics::CreateModulationDestinatio
 		return nullptr;
 	}
 
-	if (UAudioModulationDestination* NewDestination = NewObject<UAudioModulationDestination>(WorldContextObject, Name))
+	FName UniqueName = MakeUniqueObjectName(WorldContextObject, UAudioModulationDestination::StaticClass(), Name);
+	UAudioModulationDestination* NewDestination = NewObject<UAudioModulationDestination>(WorldContextObject, UAudioModulationDestination::StaticClass(), UniqueName);
+	if (NewDestination)
 	{
 		NewDestination->SetModulator(Modulator);
 		return NewDestination;
