Deprecate Activate Modulator functions

diff --git a/Engine/Plugins/Runtime/AudioModulation/Source/AudioModulation/Public/AudioModulationStatics.h b/Engine/Plugins/Runtime/AudioModulation/Source/AudioModulation/Public/AudioModulationStatics.h
--- a/Engine/Plugins/Runtime/AudioModulation/Source/AudioModulation/Public/AudioModulationStatics.h
+++ b/Engine/Plugins/Runtime/AudioModulation/Source/AudioModulation/Public/AudioModulationStatics.h
@@ -39,14 +39,14 @@ public:
 	 */
 	static AudioModulation::FAudioModulationManager* GetModulation(UWorld* World);
 
-	/** SOFT DEPRECATED: Use CreateModulationDestination and store resulting Destination instead!
-	 * Manually activates a modulation bus. If called, deactivation will only occur
+	/** Manually activates a modulation bus. If called, deactivation will only occur
 	 * if bus is manually deactivated or destroyed (i.e. will not deactivate
 	 * when all references become inactive).
 	 */
-	UFUNCTION(BlueprintCallable, Category = "Audio|Modulation", DisplayName = "Activate Control Bus (Deprecated - 5.4)", meta = (
+	UE_DEPRECATED(5.6, "To manually control the activation state of a modulator, use CreateModulationDestination.")
+	UFUNCTION(BlueprintCallable, Category = "Audio|Modulation", DisplayName = "Activate Control Bus", meta = (
 		WorldContext = "WorldContextObject",
-		Keywords = "modulation modulator")
+		DeprecatedFunction="Note", DeprecationMessage="To manually control the activation state of a modulator, use CreateModulationDestination.")
 	)
 	static void ActivateBus(const UObject* WorldContextObject, USoundControlBus* Bus);
 
@@ -61,15 +61,15 @@ public:
 	)
 	static void ActivateBusMix(const UObject* WorldContextObject, USoundControlBusMix* Mix);
 
-	/** SOFT DEPRECATED: Use CreateModulationDestination and store resulting Destination instead!
-	 * Manually activates a modulation generator. If called, deactivation will only occur
+	/** Manually activates a modulation generator. If called, deactivation will only occur
 	 * if generator is manually deactivated and not referenced or destroyed (i.e. will not deactivate
 	 * when all references become inactive).
 	 * @param Modulator - Modulator to activate
 	 */
-	UFUNCTION(BlueprintCallable, Category = "Audio|Modulation", DisplayName = "Activate Modulation Generator (Deprecated - 5.4)", meta = (
+	UE_DEPRECATED(5.6, "To manually control the activation state of a modulator, use CreateModulationDestination.")
+	UFUNCTION(BlueprintCallable, Category = "Audio|Modulation", DisplayName = "Activate Modulation Generator", meta = (
 		WorldContext = "WorldContextObject",
-		Keywords = "modulator lfo envelope follower")
+		DeprecatedFunction="Note", DeprecationMessage="To manually control the activation state of a modulator, use CreateModulationDestination.")
 	)
 	static void ActivateGenerator(const UObject* WorldContextObject, USoundModulationGenerator* Generator);
 
@@ -212,13 +212,13 @@ public:
 		FName Name,
 		USoundModulatorBase* Modulator);
 
-	/** SOFT DEPRECATED: Use CreateModulationDestination and store resulting destination instead!
-	 * Deactivates a bus. Does nothing if the provided bus is already inactive.
+	/** Deactivates a bus. Does nothing if the provided bus is already inactive.
 	 * @param Bus - Scope of modulator
 	 */
-	UFUNCTION(BlueprintCallable, Category = "Audio|Modulation", DisplayName = "Deactivate Control Bus (Deprecated - 5.4)", meta = (
-		WorldContext = "WorldContextObject", 
-		Keywords = "modulation modulator bus")
+	UE_DEPRECATED(5.6, "To manually control the activation state of a modulator, use CreateModulationDestination.")
+	UFUNCTION(BlueprintCallable, Category = "Audio|Modulation", DisplayName = "Deactivate Control Bus", meta = (
+		WorldContext = "WorldContextObject",
+		DeprecatedFunction="Note", DeprecationMessage="To manually control the activation state of a modulator, use CreateModulationDestination.")
 	)
 	static void DeactivateBus(const UObject* WorldContextObject, USoundControlBus* Bus);
 
@@ -231,14 +231,14 @@ public:
 	)
 	static void DeactivateBusMix(const UObject* WorldContextObject, USoundControlBusMix* Mix);
 
-	/** SOFT DEPRECATED: Use CreateModulationDestination and store resulting destination instead!
-	 * Deactivates a modulation generator. Does nothing if an instance of the provided generator is already inactive.
+	/** Deactivates a modulation generator. Does nothing if an instance of the provided generator is already inactive.
 	 * @param Generator - Generator to activate
 	 * @param Scope - Scope of modulator
 	 */
+	UE_DEPRECATED(5.6, "To manually control the activation state of a modulator, use CreateModulationDestination.")
 	UFUNCTION(BlueprintCallable, Category = "Audio|Modulation", DisplayName = "Deactivate Modulation Generator (Deprecated - 5.4)", meta = (
 		WorldContext = "WorldContextObject",
-		Keywords = "bus modulation modulator generator")
+		DeprecatedFunction="Note", DeprecationMessage="To manually control the activation state of a modulator, use CreateModulationDestination.")
 	)
 	static void DeactivateGenerator(const UObject* WorldContextObject, USoundModulationGenerator* Generator);
 
