Change default Modulation Patch settings for newly created Modulation Patches

diff --git a/Engine/Plugins/Runtime/AudioModulation/Source/AudioModulationEditor/Private/Factories/SoundModulationPatchFactory.cpp b/Engine/Plugins/Runtime/AudioModulation/Source/AudioModulationEditor/Private/Factories/SoundModulationPatchFactory.cpp
--- a/Engine/Plugins/Runtime/AudioModulation/Source/AudioModulationEditor/Private/Factories/SoundModulationPatchFactory.cpp
+++ b/Engine/Plugins/Runtime/AudioModulation/Source/AudioModulationEditor/Private/Factories/SoundModulationPatchFactory.cpp
@@ -19,6 +19,21 @@ USoundModulationPatchFactory::USoundModulationPatchFactory(const FObjectInitiali
 UObject* USoundModulationPatchFactory::FactoryCreateNew(UClass* Class, UObject* InParent, FName Name, EObjectFlags Flags, UObject* Context, FFeedbackContext* Warn)
 {
 	Audio::Analytics::RecordEvent_Usage(TEXT("AudioModulation.ParameterPatchCreated"));
-	return NewObject<USoundModulationPatch>(InParent, Name, Flags);
+	USoundModulationPatch* NewPatch = NewObject<USoundModulationPatch>(InParent, Name, Flags);
+	
+	if (NewPatch)
+	{
+		NewPatch->PatchSettings.bBypass = false;
+		
+		FSoundControlModulationInput DefaultInput;
+
+		DefaultInput.Transform.Curve = EWaveTableCurve::Custom;
+		DefaultInput.Transform.CurveCustom.AddKey(0.0f, 0.0f);
+		DefaultInput.Transform.CurveCustom.AddKey(1.0f, 1.0f);
+
+		NewPatch->PatchSettings.Inputs.Add(DefaultInput);
+	}
+	
+	return NewPatch;
 }
 
