[Audio Insights] - Control Buses don't clear from the Control Bus Window in some contexts
- The places currently used to send ControlBusActivate and ControlBusDeactivate trace logs do not cover all cases where control buses are activated/deactivated
- To ensure we capture all cases, these trace logs have been moved to the Initialise and Destructor of FControlBusProxy

diff --git a/Engine/Plugins/Runtime/AudioModulation/Source/AudioModulation/Private/AudioModulationSystem.cpp b/Engine/Plugins/Runtime/AudioModulation/Source/AudioModulation/Private/AudioModulationSystem.cpp
--- a/Engine/Plugins/Runtime/AudioModulation/Source/AudioModulation/Private/AudioModulationSystem.cpp
+++ b/Engine/Plugins/Runtime/AudioModulation/Source/AudioModulation/Private/AudioModulationSystem.cpp
@@ -46,12 +46,6 @@ UE_TRACE_EVENT_BEGIN(Audio, ControlBusUpdate)
 	UE_TRACE_EVENT_FIELD(float, Value)
 UE_TRACE_EVENT_END()
 
-UE_TRACE_EVENT_BEGIN(Audio, ControlBusDeactivate)
-	UE_TRACE_EVENT_FIELD(uint32, DeviceId)
-	UE_TRACE_EVENT_FIELD(uint32, ControlBusId)
-	UE_TRACE_EVENT_FIELD(double, Timestamp)
-UE_TRACE_EVENT_END()
-
 UE_TRACE_EVENT_BEGIN(Audio, BusMixRegisterBus)
 	UE_TRACE_EVENT_FIELD(uint32, DeviceId)
 	UE_TRACE_EVENT_FIELD(uint32, SourceId)
@@ -932,13 +926,6 @@ namespace AudioModulation
 			FBusHandle BusHandle = FBusHandle::Get(static_cast<FBusId>(ModId), RefProxies.Buses);
 			if (UnregisterModulator<FBusHandle>(BusHandle, RefModulators.BusMap, HandleId))
 			{
-#if UE_AUDIO_PROFILERTRACE_ENABLED
-				UE_TRACE_LOG(Audio, ControlBusDeactivate, AudioChannel)
-					<< ControlBusDeactivate.DeviceId(static_cast<uint32>(AudioDeviceId))
-					<< ControlBusDeactivate.ControlBusId(static_cast<uint32>(ModId))
-					<< ControlBusDeactivate.Timestamp(FPlatformTime::Cycles64());
-#endif
-
 				return;
 			}
 
diff --git a/Engine/Plugins/Runtime/AudioModulation/Source/AudioModulation/Private/SoundControlBusProxy.cpp b/Engine/Plugins/Runtime/AudioModulation/Source/AudioModulation/Private/SoundControlBusProxy.cpp
--- a/Engine/Plugins/Runtime/AudioModulation/Source/AudioModulation/Private/SoundControlBusProxy.cpp
+++ b/Engine/Plugins/Runtime/AudioModulation/Source/AudioModulation/Private/SoundControlBusProxy.cpp
@@ -19,6 +19,12 @@ UE_TRACE_EVENT_BEGIN(Audio, ControlBusActivate)
 	UE_TRACE_EVENT_FIELD(UE::Trace::WideString, ParamName)
 UE_TRACE_EVENT_END()
 
+UE_TRACE_EVENT_BEGIN(Audio, ControlBusDeactivate)
+	UE_TRACE_EVENT_FIELD(uint32, DeviceId)
+	UE_TRACE_EVENT_FIELD(uint32, ControlBusId)
+	UE_TRACE_EVENT_FIELD(double, Timestamp)
+UE_TRACE_EVENT_END()
+
 UE_TRACE_EVENT_BEGIN(Audio, GeneratorRegisterBus)
 	UE_TRACE_EVENT_FIELD(uint32, DeviceId)
 	UE_TRACE_EVENT_FIELD(uint32, SourceId)
@@ -45,13 +51,6 @@ namespace AudioModulation
 		FAudioModulationSystem& ModSystem = static_cast<FAudioModulationManager&>(InModulation).GetSystem();
 
 #if UE_AUDIO_PROFILERTRACE_ENABLED
-		UE_TRACE_LOG(Audio, ControlBusActivate, AudioChannel)
-			<< ControlBusActivate.DeviceId(static_cast<uint32>(ModSystem.GetAudioDeviceId()))
-			<< ControlBusActivate.ControlBusId(static_cast<uint32>(GetId()))
-			<< ControlBusActivate.Timestamp(FPlatformTime::Cycles64())
-			<< ControlBusActivate.Name(*GetName().ToString())
-			<< ControlBusActivate.ParamName(*OutputParameter.ParameterName.ToString());
-
 		for (const FModulationGeneratorSettings& GeneratorSetting : GeneratorSettings)
 		{
 			UE_TRACE_LOG(Audio, GeneratorRegisterBus, AudioChannel)
@@ -86,6 +85,19 @@ namespace AudioModulation
 		Init(MoveTemp(InSettings));
 	}
 
+	FControlBusProxy::~FControlBusProxy()
+	{
+#if UE_AUDIO_PROFILERTRACE_ENABLED
+		if (ModSystem)
+		{
+			UE_TRACE_LOG(Audio, ControlBusDeactivate, AudioChannel)
+				<< ControlBusDeactivate.DeviceId(static_cast<uint32>(ModSystem->AudioDeviceId))
+				<< ControlBusDeactivate.ControlBusId(static_cast<uint32>(GetId()))
+				<< ControlBusDeactivate.Timestamp(FPlatformTime::Cycles64());
+		}
+#endif // UE_AUDIO_PROFILERTRACE_ENABLED
+	}
+
 	FControlBusProxy& FControlBusProxy::operator =(FControlBusSettings&& InSettings)
 	{
 		Init(MoveTemp(InSettings));
@@ -151,6 +163,15 @@ namespace AudioModulation
 		// Move vs. reset and adding to original array to avoid potentially clearing handles (and thus current Generator state)
 		// and destroying generators if function is called while reinitializing/updating the modulator
 		GeneratorHandles = MoveTemp(NewHandles);
+
+#if UE_AUDIO_PROFILERTRACE_ENABLED
+		UE_TRACE_LOG(Audio, ControlBusActivate, AudioChannel)
+			<< ControlBusActivate.DeviceId(ModSystem->GetAudioDeviceId())
+			<< ControlBusActivate.ControlBusId(InSettings.GetId())
+			<< ControlBusActivate.Timestamp(FPlatformTime::Cycles64())
+			<< ControlBusActivate.Name(*(InSettings.GetName().ToString()))
+			<< ControlBusActivate.ParamName(*(GetParameterName().ToString()));
+#endif // UE_AUDIO_PROFILERTRACE_ENABLED
 	}
 
 	bool FControlBusProxy::IsBypassed() const
diff --git a/Engine/Plugins/Runtime/AudioModulation/Source/AudioModulation/Private/SoundControlBusProxy.h b/Engine/Plugins/Runtime/AudioModulation/Source/AudioModulation/Private/SoundControlBusProxy.h
--- a/Engine/Plugins/Runtime/AudioModulation/Source/AudioModulation/Private/SoundControlBusProxy.h
+++ b/Engine/Plugins/Runtime/AudioModulation/Source/AudioModulation/Private/SoundControlBusProxy.h
@@ -66,6 +66,8 @@ namespace AudioModulation
 	public:
 		FControlBusProxy();
 		FControlBusProxy(FControlBusSettings&& InSettings, FAudioModulationSystem& InModSystem);
+		
+		~FControlBusProxy();
 
 		FControlBusProxy& operator =(FControlBusSettings&& InSettings);
 
