[Audio Insights] - Generators don't clear from the Modulation Matrix window when they are no longer active

- Issue arises when a generator is removed as a result of a bus removal from the modulator proxy map (as opposed to directly removing the generator)
- If a bus has been removed, we iterate over it's generator handles and send the ModulatingSourceDeactivate trace for each of them

diff --git a/Engine/Plugins/Runtime/AudioModulation/Source/AudioModulation/Private/AudioModulationSystem.cpp b/Engine/Plugins/Runtime/AudioModulation/Source/AudioModulation/Private/AudioModulationSystem.cpp
--- a/Engine/Plugins/Runtime/AudioModulation/Source/AudioModulation/Private/AudioModulationSystem.cpp
+++ b/Engine/Plugins/Runtime/AudioModulation/Source/AudioModulation/Private/AudioModulationSystem.cpp
@@ -926,6 +926,20 @@ namespace AudioModulation
 			FBusHandle BusHandle = FBusHandle::Get(static_cast<FBusId>(ModId), RefProxies.Buses);
 			if (UnregisterModulator<FBusHandle>(BusHandle, RefModulators.BusMap, HandleId))
 			{
+#if UE_AUDIO_PROFILERTRACE_ENABLED
+				if (RefModulators.BusMap.Find(BusHandle) == nullptr)
+				{
+					const FControlBusProxy& ControlBusProxy = BusHandle.FindProxy();
+					for (const FGeneratorHandle& GeneratorHandle : ControlBusProxy.GetGeneratorHandles())
+					{
+						UE_TRACE_LOG(Audio, ModulatingSourceDeactivate, AudioChannel)
+							<< ModulatingSourceDeactivate.DeviceId(AudioDeviceId)
+							<< ModulatingSourceDeactivate.SourceId(GeneratorHandle.GetId())
+							<< ModulatingSourceDeactivate.Timestamp(FPlatformTime::Cycles64());
+					}
+				}
+#endif
+
 				return;
 			}
 
