[Audio Insights] The Sources tab does not populate with active sound sources if Audio Insights was opened after starting the PIE session

Part 01 - We now trigger Analysis when the Audio Insights window is opened during PIE (see AudioInsightsDashboardFactory)

Part 02 - Each system now sends it's current "active" state to Audio Insights when trace analysis begins.

- This change fixes the same issue across the Sources, Virtual Loops, Submixes, Audio Buses, Control Buses, and Modulation Matrix tabs
- Issue was broadly similar across all of these tabs... The trace for Audio Insights was only starting once Audio Insights was opened during PIE, so any previous messages about activation/virtualization were not being received.
- Solution is slightly different for each data type, but broadly: we detect when a trace has begun, and each system sends information to Audio Insights about it's current "active" state.
- Update includes ensuring Audio Bus and Submix Bus meters are re-initialized upon opening Audio Insights to ensure they display the correct volume data

diff --git a/Engine/Plugins/AudioInsights/Source/AudioInsightsEditor/Private/AudioInsightsEditorDashboardFactory.cpp b/Engine/Plugins/AudioInsights/Source/AudioInsightsEditor/Private/AudioInsightsEditorDashboardFactory.cpp
--- a/Engine/Plugins/AudioInsights/Source/AudioInsightsEditor/Private/AudioInsightsEditorDashboardFactory.cpp
+++ b/Engine/Plugins/AudioInsights/Source/AudioInsightsEditor/Private/AudioInsightsEditorDashboardFactory.cpp
@@ -65,21 +65,7 @@ namespace UE::Audio::Insights
 	{
 		if (InDeviceId != INDEX_NONE)
 		{
-			if (InWorld && InWorld->WorldType == EWorldType::PIE)
-			{
-				IAudioInsightsTraceModule& TraceModule = FAudioInsightsEditorModule::GetChecked().GetTraceModule();
-				TraceModule.StartTraceAnalysis(bOnlyTraceAudioChannels);
-
-				const TObjectPtr<const UAudioInsightsEditorSettings> AudioInsightsEditorSettings = GetDefault<UAudioInsightsEditorSettings>();
-
-				// We don't want to set ActiveDeviceId if bWorldFilterDefaultsToFirstClient is true and more than 2 PIE clients are running
-				if (AudioInsightsEditorSettings == nullptr ||
-					(AudioInsightsEditorSettings && !AudioInsightsEditorSettings->bWorldFilterDefaultsToFirstClient) ||
-					AudioDeviceIds.Num() < 2)
-				{
-					ActiveDeviceId = InDeviceId;
-				}
-			}
+			StartTraceAnalysis(InWorld, InDeviceId);
 		}
 
 		RefreshDeviceSelector();
@@ -223,6 +209,12 @@ namespace UE::Audio::Insights
 		RegisterTabSpawners();
 		RefreshDeviceSelector();
 
+		if (FAudioDeviceManager* DeviceManager = FAudioDeviceManager::Get();
+			DeviceManager && GEditor && GEditor->PlayWorld)
+		{
+			StartTraceAnalysis(GEditor->PlayWorld, DeviceManager->GetActiveAudioDevice()->DeviceID);
+		}
+
 		const TSharedRef<FTabManager::FLayout> TabLayout = LoadLayoutFromConfig();
 
 		const TSharedRef<SWidget> TabContent = SNew(SVerticalBox)
@@ -251,6 +243,14 @@ namespace UE::Audio::Insights
 
 		DockTab->SetOnTabClosed(SDockTab::FOnTabClosedCallback::CreateLambda([this](TSharedRef<SDockTab> TabClosed)
 		{
+			// If we are still in PIE, make sure we stop tracing for Audio Insights
+			if (FAudioDeviceManager* DeviceManager = FAudioDeviceManager::Get();
+				DeviceManager && GEditor && GEditor->PlayWorld)
+			{
+				IAudioInsightsTraceModule& TraceModule = FAudioInsightsEditorModule::GetChecked().GetTraceModule();
+				TraceModule.StopTraceAnalysis();
+			}
+
 			ResetDelegates();
 			UnregisterTabSpawners();
 			SaveLayoutToConfig();
@@ -511,6 +511,25 @@ namespace UE::Audio::Insights
 			];
 	}
 
+	void FEditorDashboardFactory::StartTraceAnalysis(const TObjectPtr<const UWorld> InWorld, const ::Audio::FDeviceId InDeviceId)
+	{
+		if (InWorld && InWorld->IsGameWorld())
+		{
+			IAudioInsightsTraceModule& TraceModule = FAudioInsightsEditorModule::GetChecked().GetTraceModule();
+			TraceModule.StartTraceAnalysis(bOnlyTraceAudioChannels);
+
+			const TObjectPtr<const UAudioInsightsEditorSettings> AudioInsightsEditorSettings = GetDefault<UAudioInsightsEditorSettings>();
+
+			// We don't want to set ActiveDeviceId if bWorldFilterDefaultsToFirstClient is true and more than 2 PIE clients are running
+			if (AudioInsightsEditorSettings == nullptr ||
+				(AudioInsightsEditorSettings && !AudioInsightsEditorSettings->bWorldFilterDefaultsToFirstClient) ||
+				AudioDeviceIds.Num() < 2)
+			{
+				ActiveDeviceId = InDeviceId;
+			}
+		}
+	}
+
 	void FEditorDashboardFactory::InitDelegates()
 	{
 		if (!OnWorldRegisteredToAudioDeviceHandle.IsValid())
diff --git a/Engine/Plugins/AudioInsights/Source/AudioInsightsEditor/Private/AudioInsightsEditorDashboardFactory.h b/Engine/Plugins/AudioInsights/Source/AudioInsightsEditor/Private/AudioInsightsEditorDashboardFactory.h
--- a/Engine/Plugins/AudioInsights/Source/AudioInsightsEditor/Private/AudioInsightsEditorDashboardFactory.h
+++ b/Engine/Plugins/AudioInsights/Source/AudioInsightsEditor/Private/AudioInsightsEditorDashboardFactory.h
@@ -44,6 +44,7 @@ namespace UE::Audio::Insights
 		TSharedRef<SWidget> MakeMenuBarWidget();
 		TSharedRef<SWidget> MakeMainToolbarWidget();
 
+		void StartTraceAnalysis(const TObjectPtr<const UWorld> InWorld, const ::Audio::FDeviceId InDeviceId);
 		void InitDelegates();
 		TSharedRef<FTabManager::FLayout> GetDefaultTabLayout();
 
diff --git a/Engine/Plugins/AudioInsights/Source/AudioInsightsEditor/Private/Providers/AudioBusProvider.cpp b/Engine/Plugins/AudioInsights/Source/AudioInsightsEditor/Private/Providers/AudioBusProvider.cpp
--- a/Engine/Plugins/AudioInsights/Source/AudioInsightsEditor/Private/Providers/AudioBusProvider.cpp
+++ b/Engine/Plugins/AudioInsights/Source/AudioInsightsEditor/Private/Providers/AudioBusProvider.cpp
@@ -6,6 +6,7 @@
 #include "AssetRegistry/AssetRegistryModule.h"
 #include "AudioInsightsEditorDashboardFactory.h"
 #include "AudioInsightsEditorModule.h"
+#include "AudioMixerTrace.h"
 #include "Sound/AudioBus.h"
 
 namespace UE::Audio::Insights
@@ -20,6 +21,10 @@ namespace UE::Audio::Insights
 		AssetRegistryModule.Get().OnFilesLoaded().AddRaw(this, &FAudioBusProvider::OnFilesLoaded);
 
 		FEditorDashboardFactory::OnActiveAudioDeviceChanged.AddRaw(this, &FAudioBusProvider::OnActiveAudioDeviceChanged);
+
+#if UE_AUDIO_PROFILERTRACE_ENABLED
+		FTraceAuxiliary::OnTraceStarted.AddRaw(this, &FAudioBusProvider::OnTraceStarted);
+#endif
 	}
 
 	FAudioBusProvider::~FAudioBusProvider()
@@ -32,6 +37,10 @@ namespace UE::Audio::Insights
 		}
 
 		FEditorDashboardFactory::OnActiveAudioDeviceChanged.RemoveAll(this);
+
+#if UE_AUDIO_PROFILERTRACE_ENABLED
+		FTraceAuxiliary::OnTraceStarted.RemoveAll(this);
+#endif // UE_AUDIO_PROFILERTRACE_ENABLED
 	}
 
 	FName FAudioBusProvider::GetName_Static()
@@ -66,6 +75,13 @@ namespace UE::Audio::Insights
 		UpdateAudioBusAssetNames();
 	}
 
+	void FAudioBusProvider::OnTraceStarted(FTraceAuxiliary::EConnectionType TraceType, const FString& TraceDestination)
+	{
+#if UE_AUDIO_PROFILERTRACE_ENABLED
+		UpdateAudioBusAssetNames();
+#endif // UE_AUDIO_PROFILERTRACE_ENABLED
+	}
+
 	void FAudioBusProvider::AddAudioBusAsset(const FAssetData& InAssetData)
 	{
 		const bool bIsAudioBusAssetAlreadAdded = AudioBusDataViewEntries.ContainsByPredicate(
@@ -90,7 +106,7 @@ namespace UE::Audio::Insights
 
 			OnAudioBusAssetAdded.Broadcast(InAssetData.GetAsset());
 
-			bAssetsUpdated = true;
+			bAssetEntriesNeedRefreshing = true;
 			++LastUpdateId;
 		}
 	}
@@ -114,7 +130,7 @@ namespace UE::Audio::Insights
 
 			OnAudioBusAssetRemoved.Broadcast(InAssetData.GetAsset());
 
-			bAssetsUpdated = true;
+			bAssetEntriesNeedRefreshing = true;
 			++LastUpdateId;
 		}
 	}
@@ -175,7 +191,7 @@ namespace UE::Audio::Insights
 
 
 		// Process messages
-		if (bAssetsUpdated)
+		if (bAssetEntriesNeedRefreshing)
 		{
 			const FAudioInsightsEditorModule AudioInsightsEditorModule = FAudioInsightsEditorModule::GetChecked();
 			const ::Audio::FDeviceId AudioDeviceId = AudioInsightsEditorModule.GetDeviceId();
@@ -194,7 +210,7 @@ namespace UE::Audio::Insights
 				}
 			}
 
-			bAssetsUpdated = false;
+			bAssetEntriesNeedRefreshing = false;
 		}
 
 		ProcessMessageQueue<FAudioBusActivateMessage>(TraceMessages.ActivateMessages, BumpEntryFunc,
@@ -238,10 +254,13 @@ namespace UE::Audio::Insights
 			{
 				FAudioBusAssetDashboardEntry& EntryRef = *OutEntry->Get();
 
-				if (EntryRef.bIsRegistered)
+				// We should only receive activity messages from registered buses
+				const bool bIsRegistered = EntryRef.bIsRegistered || (!EntryRef.Name.IsEmpty() && EntryRef.EntryType != EAudioBusEntryType::None);
+				if (bIsRegistered)
 				{
 					EntryRef.bHasActivity = Msg.bHasActivity;
 					EntryRef.Timestamp = Msg.Timestamp;
+					EntryRef.bIsRegistered = true;
 				}
 			}
 		});
@@ -324,6 +343,8 @@ namespace UE::Audio::Insights
 			TraceServices::IAnalysisSession& Session;
 		};
 
+		bAssetEntriesNeedRefreshing = true;
+
 		return new FAudioBusTraceAnalyzer(AsShared(), InSession);
 	}
 } // namespace UE::Audio::Insights
diff --git a/Engine/Plugins/AudioInsights/Source/AudioInsightsEditor/Private/Providers/AudioBusProvider.h b/Engine/Plugins/AudioInsights/Source/AudioInsightsEditor/Private/Providers/AudioBusProvider.h
--- a/Engine/Plugins/AudioInsights/Source/AudioInsightsEditor/Private/Providers/AudioBusProvider.h
+++ b/Engine/Plugins/AudioInsights/Source/AudioInsightsEditor/Private/Providers/AudioBusProvider.h
@@ -3,6 +3,7 @@
 
 #include "Delegates/Delegate.h"
 #include "Messages/AudioBusMessages.h"
+#include "ProfilingDebugging/TraceAuxiliary.h"
 #include "UObject/NameTypes.h"
 
 namespace UE::Audio::Insights
@@ -31,6 +32,7 @@ namespace UE::Audio::Insights
 		void OnAssetRemoved(const FAssetData& InAssetData);
 		void OnFilesLoaded();
 		void OnActiveAudioDeviceChanged();
+		void OnTraceStarted(FTraceAuxiliary::EConnectionType TraceType, const FString& TraceDestination);
 
 		void AddAudioBusAsset(const FAssetData& InAssetData);
 		void RemoveAudioBusAsset(const FAssetData& InAssetData);
@@ -40,7 +42,7 @@ namespace UE::Audio::Insights
 		virtual bool ProcessMessages() override;
 
 		bool bAreFilesLoaded = false;
-		bool bAssetsUpdated = false;
+		bool bAssetEntriesNeedRefreshing = false;
 
 		TArray<TSharedPtr<FAudioBusAssetDashboardEntry>> AudioBusDataViewEntries;
 
diff --git a/Engine/Plugins/AudioInsights/Source/AudioInsightsEditor/Private/Providers/SoundSubmixProvider.cpp b/Engine/Plugins/AudioInsights/Source/AudioInsightsEditor/Private/Providers/SoundSubmixProvider.cpp
--- a/Engine/Plugins/AudioInsights/Source/AudioInsightsEditor/Private/Providers/SoundSubmixProvider.cpp
+++ b/Engine/Plugins/AudioInsights/Source/AudioInsightsEditor/Private/Providers/SoundSubmixProvider.cpp
@@ -6,6 +6,7 @@
 #include "AssetRegistry/AssetRegistryModule.h"
 #include "AudioInsightsEditorDashboardFactory.h"
 #include "AudioInsightsEditorModule.h"
+#include "AudioMixerTrace.h"
 #include "Sound/SoundSubmix.h"
 
 namespace UE::Audio::Insights
@@ -20,6 +21,10 @@ namespace UE::Audio::Insights
 		AssetRegistryModule.Get().OnFilesLoaded().AddRaw(this, &FSoundSubmixProvider::OnFilesLoaded);
 
 		FEditorDashboardFactory::OnActiveAudioDeviceChanged.AddRaw(this, &FSoundSubmixProvider::OnActiveAudioDeviceChanged);
+		
+#if UE_AUDIO_PROFILERTRACE_ENABLED
+		FTraceAuxiliary::OnTraceStarted.AddRaw(this, &FSoundSubmixProvider::OnTraceStarted);
+#endif // UE_AUDIO_PROFILERTRACE_ENABLED
 	}
 	
 	FSoundSubmixProvider::~FSoundSubmixProvider()
@@ -32,6 +37,10 @@ namespace UE::Audio::Insights
 		}
 
 		FEditorDashboardFactory::OnActiveAudioDeviceChanged.RemoveAll(this);
+
+#if UE_AUDIO_PROFILERTRACE_ENABLED
+		FTraceAuxiliary::OnTraceStarted.RemoveAll(this);
+#endif // UE_AUDIO_PROFILERTRACE_ENABLED
 	}
 
 	FName FSoundSubmixProvider::GetName_Static()
@@ -66,6 +75,13 @@ namespace UE::Audio::Insights
 		UpdateSubmixAssetNames();
 	}
 
+	void FSoundSubmixProvider::OnTraceStarted(FTraceAuxiliary::EConnectionType TraceType, const FString& TraceDestination)
+	{
+#if UE_AUDIO_PROFILERTRACE_ENABLED
+		UpdateSubmixAssetNames();
+#endif // UE_AUDIO_PROFILERTRACE_ENABLED
+	}
+
 	void FSoundSubmixProvider::AddSubmixAsset(const FAssetData& InAssetData)
 	{
 		const bool bIsSubmixAssetAlreadAdded = SubmixDataViewEntries.ContainsByPredicate(
@@ -89,7 +105,7 @@ namespace UE::Audio::Insights
 
 			OnSubmixAssetAdded.Broadcast(InAssetData.GetAsset());
 
-			bAssetsUpdated = true;
+			bAssetEntriesNeedRefreshing = true;
 			++LastUpdateId;
 		}
 	}
@@ -113,7 +129,7 @@ namespace UE::Audio::Insights
 
 			OnSubmixAssetRemoved.Broadcast(InAssetData.GetAsset());
 
-			bAssetsUpdated = true;
+			bAssetEntriesNeedRefreshing = true;
 			++LastUpdateId;
 		}
 	}
@@ -153,7 +169,7 @@ namespace UE::Audio::Insights
 
 
 		// Process messages
-		if (bAssetsUpdated)
+		if (bAssetEntriesNeedRefreshing)
 		{
 			const FAudioInsightsEditorModule AudioInsightsEditorModule = FAudioInsightsEditorModule::GetChecked();
 			const ::Audio::FDeviceId AudioDeviceId = AudioInsightsEditorModule.GetDeviceId();
@@ -172,7 +188,7 @@ namespace UE::Audio::Insights
 				}
 			}
 
-			bAssetsUpdated = false;
+			bAssetEntriesNeedRefreshing = false;
 		}
 
 		ProcessMessageQueue<FSoundSubmixHasActivityMessage>(TraceMessages.HasActivityMessages, GetEntry,
@@ -248,6 +264,8 @@ namespace UE::Audio::Insights
 			TraceServices::IAnalysisSession& Session;
 		};
 
+		bAssetEntriesNeedRefreshing = true;
+
 		return new FSoundSubmixTraceAnalyzer(AsShared(), InSession);
 	}
 } // namespace UE::Audio::Insights
diff --git a/Engine/Plugins/AudioInsights/Source/AudioInsightsEditor/Private/Providers/SoundSubmixProvider.h b/Engine/Plugins/AudioInsights/Source/AudioInsightsEditor/Private/Providers/SoundSubmixProvider.h
--- a/Engine/Plugins/AudioInsights/Source/AudioInsightsEditor/Private/Providers/SoundSubmixProvider.h
+++ b/Engine/Plugins/AudioInsights/Source/AudioInsightsEditor/Private/Providers/SoundSubmixProvider.h
@@ -3,6 +3,7 @@
 
 #include "Delegates/Delegate.h"
 #include "Messages/SoundSubmixMessages.h"
+#include "ProfilingDebugging/TraceAuxiliary.h"
 #include "UObject/NameTypes.h"
 
 namespace UE::Audio::Insights
@@ -31,6 +32,7 @@ namespace UE::Audio::Insights
 		void OnAssetRemoved(const FAssetData& InAssetData);
 		void OnFilesLoaded();
 		void OnActiveAudioDeviceChanged();
+		void OnTraceStarted(FTraceAuxiliary::EConnectionType TraceType, const FString& TraceDestination);
 
 		void AddSubmixAsset(const FAssetData& InAssetData);
 		void RemoveSubmixAsset(const FAssetData& InAssetData);
@@ -40,7 +42,7 @@ namespace UE::Audio::Insights
 		virtual bool ProcessMessages() override;
 
 		bool bAreFilesLoaded = false;
-		bool bAssetsUpdated  = false;
+		bool bAssetEntriesNeedRefreshing  = false;
 
 		TArray<TSharedPtr<FSoundSubmixAssetDashboardEntry>> SubmixDataViewEntries;
 
diff --git a/Engine/Plugins/AudioInsights/Source/AudioInsightsEditor/Private/Views/AudioBusesDashboardViewFactory.cpp b/Engine/Plugins/AudioInsights/Source/AudioInsightsEditor/Private/Views/AudioBusesDashboardViewFactory.cpp
--- a/Engine/Plugins/AudioInsights/Source/AudioInsightsEditor/Private/Views/AudioBusesDashboardViewFactory.cpp
+++ b/Engine/Plugins/AudioInsights/Source/AudioInsightsEditor/Private/Views/AudioBusesDashboardViewFactory.cpp
@@ -105,7 +105,8 @@ namespace UE::Audio::Insights
 					return SNullWidget::NullWidget;
 				}
 
-				AudioBusCheckboxCheckedStates.FindOrAdd(AudioBus, false);
+				const bool bInitValue = AudioBusCheckboxCheckedStates.FindOrAdd(AudioBus, false);
+				OnBusAssetInit.Broadcast(bInitValue, AudioBus);
 
 				return SNew(SHorizontalBox)
 					.Clipping(EWidgetClipping::ClipToBounds)
@@ -295,6 +296,11 @@ namespace UE::Audio::Insights
 			}
 		}
 
+		for (const auto& [AudioBus, bIsCheckedState] : AudioBusCheckboxCheckedStates)
+		{
+			OnBusAssetInit.Broadcast(bIsCheckedState, AudioBus);
+		}
+
 		return DashboardWidget->AsShared();
 	}
 
diff --git a/Engine/Plugins/AudioInsights/Source/AudioInsightsEditor/Private/Views/AudioBusesDashboardViewFactory.h b/Engine/Plugins/AudioInsights/Source/AudioInsightsEditor/Private/Views/AudioBusesDashboardViewFactory.h
--- a/Engine/Plugins/AudioInsights/Source/AudioInsightsEditor/Private/Views/AudioBusesDashboardViewFactory.h
+++ b/Engine/Plugins/AudioInsights/Source/AudioInsightsEditor/Private/Views/AudioBusesDashboardViewFactory.h
@@ -21,6 +21,9 @@ namespace UE::Audio::Insights
 		virtual EDefaultDashboardTabStack GetDefaultTabStack() const override;
 		virtual TSharedRef<SWidget> MakeWidget(TSharedRef<SDockTab> OwnerTab, const FSpawnTabArgs& SpawnTabArgs) override;
 
+		DECLARE_MULTICAST_DELEGATE_TwoParams(FOnBusAssetInit, const bool /*bIsChecked*/, TWeakObjectPtr<UAudioBus> /*AudioBus*/);
+		inline static FOnBusAssetInit OnBusAssetInit;
+
 		DECLARE_MULTICAST_DELEGATE_TwoParams(FOnAudioBusAssetChecked, const bool /*bIsChecked*/, const TWeakObjectPtr<UAudioBus> /*AudioBus*/);
 		inline static FOnAudioBusAssetChecked OnAudioBusAssetChecked;
 
diff --git a/Engine/Plugins/AudioInsights/Source/AudioInsightsEditor/Private/Views/AudioMetersDashboardViewFactory.cpp b/Engine/Plugins/AudioInsights/Source/AudioInsightsEditor/Private/Views/AudioMetersDashboardViewFactory.cpp
--- a/Engine/Plugins/AudioInsights/Source/AudioInsightsEditor/Private/Views/AudioMetersDashboardViewFactory.cpp
+++ b/Engine/Plugins/AudioInsights/Source/AudioInsightsEditor/Private/Views/AudioMetersDashboardViewFactory.cpp
@@ -36,6 +36,16 @@ namespace UE::Audio::Insights
 
 	TSharedRef<SWidget> FAudioMetersDashboardViewFactory::MakeWidget(TSharedRef<SDockTab> OwnerTab, const FSpawnTabArgs& SpawnTabArgs)
 	{
+		if (!OnSubmixAssetInit.IsValid())
+		{
+			OnSubmixAssetInit = FSubmixesDashboardViewFactory::OnSubmixAssetInit.AddSP(this, &FAudioMetersDashboardViewFactory::HandleOnSubmixAssetInit);
+		}
+
+		if (!OnAudioBusAssetInit.IsValid())
+		{
+			OnAudioBusAssetInit = FAudioBusesDashboardViewFactory::OnBusAssetInit.AddSP(this, &FAudioMetersDashboardViewFactory::HandleOnAudioBusAssetInit);
+		}
+
 		if (!OnSubmixAssetCheckedHandle.IsValid())
 		{
 			OnSubmixAssetCheckedHandle = FSubmixesDashboardViewFactory::OnSubmixAssetChecked.AddSP(this, &FAudioMetersDashboardViewFactory::HandleOnSubmixAssetChecked);
@@ -69,6 +79,71 @@ namespace UE::Audio::Insights
 		return MeterViewsScrollBox.ToSharedRef();
 	}
 
+	void FAudioMetersDashboardViewFactory::HandleOnSubmixAssetInit(const bool bInIsChecked, const TWeakObjectPtr<USoundSubmix> InSoundSubmix)
+	{
+		if (!bInIsChecked || !InSoundSubmix.IsValid())
+		{
+			return;
+		}
+
+		const uint32 SubmixUniqueId = InSoundSubmix->GetUniqueID();
+
+		if (!AudioMeterViews.Contains(SubmixUniqueId))
+		{
+			return;
+		}
+
+		// Tear down and re-create meter analyzers on initialization to ensure we receive the correct volume data
+		const int32 SlotIndex = AudioMeterViewsContainer->RemoveSlot(AudioMeterViews[SubmixUniqueId]->GetWidget());
+
+		AudioMeterViews.Remove(SubmixUniqueId);
+		AudioMeterViews.Add(SubmixUniqueId, MakeShared<FAudioMeterView>(FAudioMeterView::FAudioAssetVariant(TInPlaceType<TWeakObjectPtr<USoundSubmix>>(), InSoundSubmix)));
+
+		AudioMeterViewsContainer->InsertSlot(SlotIndex)
+		.AutoWidth()
+		.Padding(10.0f, 0.0f, 10.0f, 0.0f)
+		[
+			AudioMeterViews[SubmixUniqueId]->GetWidget()
+		];
+
+		if (MeterViewsScrollBox.IsValid())
+		{
+			MeterViewsScrollBox->Invalidate(EInvalidateWidget::Layout);
+		}
+	}
+
+	void FAudioMetersDashboardViewFactory::HandleOnAudioBusAssetInit(const bool bInIsChecked, const TWeakObjectPtr<UAudioBus> InAudioBus)
+	{
+		if (!bInIsChecked || !InAudioBus.IsValid())
+		{
+			return;
+		}
+
+		const uint32 AudioBusUniqueId = InAudioBus->GetUniqueID();
+		if (!AudioMeterViews.Contains(AudioBusUniqueId))
+		{
+			return;
+		}
+
+		// Tear down and re-create meter analyzers on initialization to ensure we receive the correct volume data
+		const int32 SlotIndex = AudioMeterViewsContainer->RemoveSlot(AudioMeterViews[AudioBusUniqueId]->GetWidget());
+
+		AudioMeterViews.Remove(AudioBusUniqueId);
+		AudioMeterViews.Add(AudioBusUniqueId, MakeShared<FAudioMeterView>(FAudioMeterView::FAudioAssetVariant(TInPlaceType<TWeakObjectPtr<UAudioBus>>(), InAudioBus)));
+
+		AudioMeterViewsContainer->InsertSlot(SlotIndex)
+		.AutoWidth()
+		.Padding(10.0f, 0.0f, 10.0f, 0.0f)
+		[
+			AudioMeterViews[AudioBusUniqueId]->GetWidget()
+		];
+
+		if (MeterViewsScrollBox.IsValid())
+		{
+			MeterViewsScrollBox->Invalidate(EInvalidateWidget::Layout);
+		}
+	}
+
 	void FAudioMetersDashboardViewFactory::HandleOnSubmixAssetChecked(const bool bInIsChecked, const TWeakObjectPtr<USoundSubmix> InSoundSubmix)
 	{
 		if (!InSoundSubmix.IsValid())
diff --git a/Engine/Plugins/AudioInsights/Source/AudioInsightsEditor/Private/Views/AudioMetersDashboardViewFactory.h b/Engine/Plugins/AudioInsights/Source/AudioInsightsEditor/Private/Views/AudioMetersDashboardViewFactory.h
--- a/Engine/Plugins/AudioInsights/Source/AudioInsightsEditor/Private/Views/AudioMetersDashboardViewFactory.h
+++ b/Engine/Plugins/AudioInsights/Source/AudioInsightsEditor/Private/Views/AudioMetersDashboardViewFactory.h
@@ -26,6 +26,9 @@ namespace UE::Audio::Insights
 		virtual TSharedRef<SWidget> MakeWidget(TSharedRef<SDockTab> OwnerTab, const FSpawnTabArgs& SpawnTabArgs) override;
 
 	private:
+		void HandleOnSubmixAssetInit(const bool bInIsChecked, const TWeakObjectPtr<USoundSubmix> InSoundSubmix);
+		void HandleOnAudioBusAssetInit(const bool bInIsChecked, const TWeakObjectPtr<UAudioBus> InAudioBus);
+
 		void HandleOnSubmixAssetChecked(const bool bInIsChecked, const TWeakObjectPtr<USoundSubmix> InSoundSubmix);
 		void HandleOnAudioBusAssetChecked(const bool bInIsChecked, const TWeakObjectPtr<UAudioBus> InAudioBus);
 
@@ -35,6 +38,8 @@ namespace UE::Audio::Insights
 		TSharedPtr<SHorizontalBox> AudioMeterViewsContainer;
 		TMap<uint32, TSharedRef<FAudioMeterView>> AudioMeterViews;
 
+		FDelegateHandle OnSubmixAssetInit;
+		FDelegateHandle OnAudioBusAssetInit;
 		FDelegateHandle OnSubmixAssetCheckedHandle;
 		FDelegateHandle OnAudioBusAssetCheckedHandle;
 		FDelegateHandle OnSubmixAssetRemovedHandle;
diff --git a/Engine/Plugins/AudioInsights/Source/AudioInsightsEditor/Private/Views/SubmixesDashboardViewFactory.cpp b/Engine/Plugins/AudioInsights/Source/AudioInsightsEditor/Private/Views/SubmixesDashboardViewFactory.cpp
--- a/Engine/Plugins/AudioInsights/Source/AudioInsightsEditor/Private/Views/SubmixesDashboardViewFactory.cpp
+++ b/Engine/Plugins/AudioInsights/Source/AudioInsightsEditor/Private/Views/SubmixesDashboardViewFactory.cpp
@@ -99,7 +99,8 @@ namespace UE::Audio::Insights
 				return SNullWidget::NullWidget;
 			}
 
-			SubmixCheckboxCheckedStates.FindOrAdd(SoundSubmix, false);
+			const bool bInitValue = SubmixCheckboxCheckedStates.FindOrAdd(SoundSubmix, false);
+			OnSubmixAssetInit.Broadcast(bInitValue, SoundSubmix);
 
 			return SNew(SHorizontalBox)
 				.Clipping(EWidgetClipping::ClipToBounds)
@@ -194,6 +195,11 @@ namespace UE::Audio::Insights
 			}
 		}
 
+		for (const auto& [AudioBus, bIsCheckedState] : SubmixCheckboxCheckedStates)
+		{
+			OnSubmixAssetInit.Broadcast(bIsCheckedState, AudioBus);
+		}
+
 		return DashboardWidget->AsShared();
 	}
 
diff --git a/Engine/Plugins/AudioInsights/Source/AudioInsightsEditor/Private/Views/SubmixesDashboardViewFactory.h b/Engine/Plugins/AudioInsights/Source/AudioInsightsEditor/Private/Views/SubmixesDashboardViewFactory.h
--- a/Engine/Plugins/AudioInsights/Source/AudioInsightsEditor/Private/Views/SubmixesDashboardViewFactory.h
+++ b/Engine/Plugins/AudioInsights/Source/AudioInsightsEditor/Private/Views/SubmixesDashboardViewFactory.h
@@ -22,6 +22,9 @@ namespace UE::Audio::Insights
 		virtual EDefaultDashboardTabStack GetDefaultTabStack() const override;
 		virtual TSharedRef<SWidget> MakeWidget(TSharedRef<SDockTab> OwnerTab, const FSpawnTabArgs& SpawnTabArgs) override;
 
+		DECLARE_MULTICAST_DELEGATE_TwoParams(FOnSubmixAssetInit, const bool /*bIsChecked*/, TWeakObjectPtr<USoundSubmix> /*SoundSubmix*/);
+		inline static FOnSubmixAssetInit OnSubmixAssetInit;
+
 		DECLARE_MULTICAST_DELEGATE_TwoParams(FOnSubmixAssetChecked, const bool /*bIsChecked*/, TWeakObjectPtr<USoundSubmix> /*SoundSubmix*/);
 		inline static FOnSubmixAssetChecked OnSubmixAssetChecked;
 
diff --git a/Engine/Plugins/Runtime/AudioModulation/Source/AudioModulation/Private/AudioModulationSystem.cpp b/Engine/Plugins/Runtime/AudioModulation/Source/AudioModulation/Private/AudioModulationSystem.cpp
--- a/Engine/Plugins/Runtime/AudioModulation/Source/AudioModulation/Private/AudioModulationSystem.cpp
+++ b/Engine/Plugins/Runtime/AudioModulation/Source/AudioModulation/Private/AudioModulationSystem.cpp
@@ -103,12 +103,23 @@ namespace AudioModulation
 		COUNT
 	};
 
+	FAudioModulationSystem::~FAudioModulationSystem()
+	{
+#if UE_AUDIO_PROFILERTRACE_ENABLED
+		FTraceAuxiliary::OnTraceStarted.RemoveAll(this);
+#endif // UE_AUDIO_PROFILERTRACE_ENABLED
+	}
+
 	void FAudioModulationSystem::Initialize(const FAudioPluginInitializationParams& InitializationParams)
 	{
 #if !UE_BUILD_SHIPPING
 		Debugger = MakeShared<FAudioModulationDebugger>();
 #endif // !UE_BUILD_SHIPPING
 
+#if UE_AUDIO_PROFILERTRACE_ENABLED
+		FTraceAuxiliary::OnTraceStarted.AddRaw(this, &FAudioModulationSystem::OnTraceStarted);
+#endif // UE_AUDIO_PROFILERTRACE_ENABLED
+
 		AudioDeviceId = InitializationParams.AudioDevicePtr->DeviceID;
 	}
 
@@ -824,10 +835,43 @@ namespace AudioModulation
 		}
 	}
 
+	void FAudioModulationSystem::OnTraceStarted(FTraceAuxiliary::EConnectionType TraceType, const FString& TraceDestination)
+	{
+#if UE_AUDIO_PROFILERTRACE_ENABLED
+		RunCommandOnProcessingThread([this]()
+		{
+			for (const auto& [BusId, ControlBusProxy] : RefProxies.Buses)
+			{
+				UE_TRACE_LOG(Audio, BusMixRegisterBus, AudioChannel)
+					<< BusMixRegisterBus.DeviceId(AudioDeviceId)
+					<< BusMixRegisterBus.SourceId(ControlBusProxy.GetId())
+					<< BusMixRegisterBus.Timestamp(FPlatformTime::Cycles64())
+					<< BusMixRegisterBus.ModulatingSourceId(BusId)
+					<< BusMixRegisterBus.BusName(*(ControlBusProxy.GetName().ToString()));
+			}
+
+			for (const auto& [BusMixId, ModulatorBusMixProxy] : RefProxies.BusMixes)
+			{
+				UE_TRACE_LOG(Audio, BusMixActivate, AudioChannel)
+					<< BusMixActivate.DeviceId(AudioDeviceId)
+					<< BusMixActivate.SourceId(BusMixId)
+					<< BusMixActivate.Timestamp(FPlatformTime::Cycles64())
+					<< BusMixActivate.Name(*(ModulatorBusMixProxy.GetName().ToString()));
+			}
+		});
+#endif // UE_AUDIO_PROFILERTRACE_ENABLED
+	}
+
 	Audio::FModulatorTypeId FAudioModulationSystem::RegisterModulator(Audio::FModulatorHandleId InHandleId, const FControlBusSettings& InSettings)
 	{
 		FControlBusSettings CachedSettings = InSettings;
-		RegisterModulator(InHandleId, MoveTemp(CachedSettings), RefProxies.Buses, RefModulators.BusMap);
+		RegisterModulator(InHandleId, MoveTemp(CachedSettings), RefProxies.Buses, RefModulators.BusMap, [this](FBusHandle& NewHandle)
+		{
+#if UE_AUDIO_PROFILERTRACE_ENABLED
+			NewHandle.FindProxy().RegisterForTraceMessages();
+#endif // #if UE_AUDIO_PROFILERTRACE_ENABLED
+		});
+
 		return static_cast<Audio::FModulatorTypeId>(EModulatorType::Bus);
 	}
 
@@ -845,6 +889,7 @@ namespace AudioModulation
 	{
 		FModulationPatchSettings CachedSettings = InSettings;
 		RegisterModulator(InHandleId, MoveTemp(CachedSettings), RefProxies.Patches, RefModulators.PatchMap);
+
 		return static_cast<Audio::FModulatorTypeId>(EModulatorType::Patch);
 	}
 
@@ -929,7 +974,9 @@ namespace AudioModulation
 #if UE_AUDIO_PROFILERTRACE_ENABLED
 				if (RefModulators.BusMap.Find(BusHandle) == nullptr)
 				{
-					const FControlBusProxy& ControlBusProxy = BusHandle.FindProxy();
+					FControlBusProxy& ControlBusProxy = BusHandle.FindProxy();
+					ControlBusProxy.UnregisterForTraceMessages();
+
 					for (const FGeneratorHandle& GeneratorHandle : ControlBusProxy.GetGeneratorHandles())
 					{
 						UE_TRACE_LOG(Audio, ModulatingSourceDeactivate, AudioChannel)
diff --git a/Engine/Plugins/Runtime/AudioModulation/Source/AudioModulation/Private/AudioModulationSystem.h b/Engine/Plugins/Runtime/AudioModulation/Source/AudioModulation/Private/AudioModulationSystem.h
--- a/Engine/Plugins/Runtime/AudioModulation/Source/AudioModulation/Private/AudioModulationSystem.h
+++ b/Engine/Plugins/Runtime/AudioModulation/Source/AudioModulation/Private/AudioModulationSystem.h
@@ -8,6 +8,7 @@
 #include "CoreMinimal.h"
 #include "HAL/CriticalSection.h"
 #include "IAudioModulation.h"
+#include "ProfilingDebugging/TraceAuxiliary.h"
 #include "SoundControlBus.h"
 #include "SoundControlBusMix.h"
 #include "SoundControlBusMixProxy.h"
@@ -66,6 +67,8 @@ namespace AudioModulation
 	class FAudioModulationSystem : public FGCObject
 	{
 	public:
+		~FAudioModulationSystem();
+
 		void Initialize(const FAudioPluginInitializationParams& InitializationParams);
 
 		UE_DEPRECATED(5.4, "Activation of modulators in this manner is now deprecated. Use USoundModulationDestinations to safely activate and track a given modulator")
@@ -171,6 +174,8 @@ namespace AudioModulation
 		/* Runs the provided command on the audio render thread (at the beginning of the ProcessModulators call) */
 		void RunCommandOnProcessingThread(TUniqueFunction<void()> Cmd);
 
+		void OnTraceStarted(FTraceAuxiliary::EConnectionType TraceType, const FString& TraceDestination);
+
 		/* Template for register calls that move the modulator settings objects onto the Audio Processing Thread & create proxies.
 		 * If the proxy is already found, adds provided HandleId as reference to given proxy. Does *not* update the proxy with the
 		 * given settings.  If update is desired on an existing proxy, UpdateModulator must be explicitly called.
@@ -184,7 +189,7 @@ namespace AudioModulation
 		 * @param InInitFunction Function type used to call on the AudioProcessingThread immediately following proxy generation
 		 */
 		template <typename THandleType, typename TModSettings, typename TMapType, typename TInitFunction = TUniqueFunction<void(THandleType)>>
-		void RegisterModulator(Audio::FModulatorHandleId InHandleId, TModSettings&& InModSettings, TMapType& OutProxyMap, TMap<THandleType, FModulatorHandleSet>& OutModMap, TInitFunction InInitFunction = TInitFunction())
+		void RegisterModulator(Audio::FModulatorHandleId InHandleId, TModSettings&& InModSettings, TMapType& OutProxyMap, TMap<THandleType, FModulatorHandleSet>& OutModMap, TInitFunction InInitFunction = [](THandleType){ })
 		{
 			check(InHandleId != INDEX_NONE);
 
@@ -202,6 +207,8 @@ namespace AudioModulation
 
 				THandleType Handle = THandleType::Create(MoveTemp(ModSettings), *PassedProxyMap, *this);
 				PassedModMap->FindOrAdd(Handle).Add(InHandleId);
+
+				InitFunction(Handle);
 			});
 		}
 
@@ -290,6 +297,8 @@ namespace AudioModulation
 	class FAudioModulationSystem
 	{
 	public:
+		~FAudioModulationSystem();
+
 		void Initialize(const FAudioPluginInitializationParams& InitializationParams) { }
 
 		void OnAuditionEnd() { }
diff --git a/Engine/Plugins/Runtime/AudioModulation/Source/AudioModulation/Private/SoundControlBusProxy.cpp b/Engine/Plugins/Runtime/AudioModulation/Source/AudioModulation/Private/SoundControlBusProxy.cpp
--- a/Engine/Plugins/Runtime/AudioModulation/Source/AudioModulation/Private/SoundControlBusProxy.cpp
+++ b/Engine/Plugins/Runtime/AudioModulation/Source/AudioModulation/Private/SoundControlBusProxy.cpp
@@ -193,6 +193,34 @@ namespace AudioModulation
 		return OutValue;
 	}
 
+	void AudioModulation::FControlBusProxy::OnTraceStarted(FTraceAuxiliary::EConnectionType TraceType, const FString& TraceDestination)
+	{
+#if UE_AUDIO_PROFILERTRACE_ENABLED
+		UE_TRACE_LOG(Audio, ControlBusActivate, AudioChannel)
+			<< ControlBusActivate.DeviceId(ModSystem->GetAudioDeviceId())
+			<< ControlBusActivate.ControlBusId(GetId())
+			<< ControlBusActivate.Timestamp(FPlatformTime::Cycles64())
+			<< ControlBusActivate.Name(*(GetName().ToString()))
+			<< ControlBusActivate.ParamName(*(GetParameterName().ToString()));
+
+		for (const FGeneratorHandle& GeneratorHandle: GeneratorHandles)
+		{
+			UE_TRACE_LOG(Audio, GeneratorRegisterBus, AudioChannel)
+				<< GeneratorRegisterBus.DeviceId(ModSystem->GetAudioDeviceId())
+				<< GeneratorRegisterBus.SourceId(GetId())
+				<< GeneratorRegisterBus.Timestamp(FPlatformTime::Cycles64())
+				<< GeneratorRegisterBus.ModulatingSourceId(GeneratorHandle.GetId())
+				<< GeneratorRegisterBus.BusName(*(GetName().ToString()));
+
+			UE_TRACE_LOG(Audio, GeneratorActivate, AudioChannel)
+				<< GeneratorActivate.DeviceId(ModSystem->GetAudioDeviceId())
+				<< GeneratorActivate.SourceId(GeneratorHandle.GetId())
+				<< GeneratorActivate.Timestamp(FPlatformTime::Cycles64())
+				<< GeneratorActivate.Name(*(GeneratorHandle.FindProxy().GetName().ToString()));
+		}
+#endif // UE_AUDIO_PROFILERTRACE_ENABLED
+	}
+
 	void FControlBusProxy::MixIn(const float InValue)
 	{
 		MixValue = Mix(InValue);
@@ -218,4 +246,18 @@ namespace AudioModulation
 		GeneratorValue = 1.0f;
 		MixValue = NAN;
 	}
+
+	void AudioModulation::FControlBusProxy::RegisterForTraceMessages()
+	{
+#if UE_AUDIO_PROFILERTRACE_ENABLED
+		FTraceAuxiliary::OnTraceStarted.AddRaw(this, &FControlBusProxy::OnTraceStarted);
+#endif // UE_AUDIO_PROFILERTRACE_ENABLED
+	}
+
+	void AudioModulation::FControlBusProxy::UnregisterForTraceMessages()
+	{
+#if UE_AUDIO_PROFILERTRACE_ENABLED
+		FTraceAuxiliary::OnTraceStarted.RemoveAll(this);
+#endif // UE_AUDIO_PROFILERTRACE_ENABLED
+	}
 } // namespace AudioModulation
diff --git a/Engine/Plugins/Runtime/AudioModulation/Source/AudioModulation/Private/SoundControlBusProxy.h b/Engine/Plugins/Runtime/AudioModulation/Source/AudioModulation/Private/SoundControlBusProxy.h
--- a/Engine/Plugins/Runtime/AudioModulation/Source/AudioModulation/Private/SoundControlBusProxy.h
+++ b/Engine/Plugins/Runtime/AudioModulation/Source/AudioModulation/Private/SoundControlBusProxy.h
@@ -3,6 +3,7 @@
 
 #include "AudioDeviceManager.h"
 #include "IAudioModulation.h"
+#include "ProfilingDebugging/TraceAuxiliary.h"
 #include "SoundControlBus.h"
 #include "SoundModulationParameter.h"
 #include "SoundModulationProxy.h"
@@ -82,10 +83,15 @@ namespace AudioModulation
 		void MixGenerators();
 		void Reset();
 
+		void RegisterForTraceMessages();
+		void UnregisterForTraceMessages();
+
 	private:
 		void Init(FControlBusSettings&& InSettings);
 		float Mix(float ValueA) const;
 
+		void OnTraceStarted(FTraceAuxiliary::EConnectionType TraceType, const FString& TraceDestination);
+
 		float DefaultValue;
 
 		// Cached values
diff --git a/Engine/Source/Runtime/AudioMixer/Private/AudioMixerSource.cpp b/Engine/Source/Runtime/AudioMixer/Private/AudioMixerSource.cpp
--- a/Engine/Source/Runtime/AudioMixer/Private/AudioMixerSource.cpp
+++ b/Engine/Source/Runtime/AudioMixer/Private/AudioMixerSource.cpp
@@ -520,6 +520,10 @@ namespace Audio
 
 	FMixerSource::~FMixerSource()
 	{
+#if UE_AUDIO_PROFILERTRACE_ENABLED
+		FTraceAuxiliary::OnTraceStarted.RemoveAll(this);
+#endif // UE_AUDIO_PROFILERTRACE_ENABLED
+
 		FreeResources();
 	}
 
@@ -884,6 +888,10 @@ namespace Audio
 			}
 #endif
 
+#if UE_AUDIO_PROFILERTRACE_ENABLED
+			FTraceAuxiliary::OnTraceStarted.AddRaw(this, &FMixerSource::OnTraceStarted);
+#endif // UE_AUDIO_PROFILERTRACE_ENABLED
+
 			// Pass the decompression state off to the mixer source buffer if it hasn't already done so
 			ICompressedAudioInfo* Decoder = MixerBuffer->GetDecompressionState(true);
 			MixerSourceBuffer->SetDecoder(Decoder);
@@ -2130,6 +2138,34 @@ namespace Audio
 	}
 #endif // if ENABLE_AUDIO_DEBUG
 
+#if UE_AUDIO_PROFILERTRACE_ENABLED
+	void FMixerSource::OnTraceStarted(FTraceAuxiliary::EConnectionType TraceType, const FString& TraceDestination)
+	{
+		// Make sure we send mixer source data to Audio Insights if it's opened in the middle of a PIE session
+		const bool bChannelEnabled = UE_TRACE_CHANNELEXPR_IS_ENABLED(AudioMixerChannel);
+		if (bChannelEnabled && WaveInstance)
+		{
+			if (const FActiveSound* ActiveSound = WaveInstance->ActiveSound;
+				ActiveSound && ActiveSound->IsPlayingAudio())
+			{
+				int32 TraceSourceId = INDEX_NONE;
+				if (MixerSourceVoice)
+				{
+					TraceSourceId = MixerSourceVoice->GetSourceId();
+				}
+
+				UE_TRACE_LOG(Audio, MixerSourceStart, AudioMixerChannel)
+					<< MixerSourceStart.DeviceId(MixerDevice->DeviceID)
+					<< MixerSourceStart.Timestamp(FPlatformTime::Cycles64())
+					<< MixerSourceStart.PlayOrder(WaveInstance->GetPlayOrder())
+					<< MixerSourceStart.SourceId(TraceSourceId)
+					<< MixerSourceStart.ComponentId(ActiveSound->GetAudioComponentID())
+					<< MixerSourceStart.Name(*WaveInstance->WaveData->GetPathName());
+			}
+		}
+	}
+#endif // UE_AUDIO_PROFILERTRACE_ENABLED
+
 	bool FMixerSource::ComputeMonoChannelMap(Audio::FAlignedFloatBuffer& OutChannelMap)
 	{
 		if (IsUsingObjectBasedSpatialization())
diff --git a/Engine/Source/Runtime/AudioMixer/Private/AudioMixerSource.h b/Engine/Source/Runtime/AudioMixer/Private/AudioMixerSource.h
--- a/Engine/Source/Runtime/AudioMixer/Private/AudioMixerSource.h
+++ b/Engine/Source/Runtime/AudioMixer/Private/AudioMixerSource.h
@@ -5,6 +5,7 @@
 #include "CoreMinimal.h"
 #include "AudioMixerBuffer.h"
 #include "AudioMixerSourceManager.h"
+#include "ProfilingDebugging/TraceAuxiliary.h"
 
 namespace Audio
 {
@@ -97,6 +98,10 @@ namespace Audio
 		void UpdateCPUCoreUtilization();
 #endif // ENABLE_AUDIO_DEBUG
 
+#if UE_AUDIO_PROFILERTRACE_ENABLED
+		void OnTraceStarted(FTraceAuxiliary::EConnectionType TraceType, const FString& TraceDestination);
+#endif // UE_AUDIO_PROFILERTRACE_ENABLED
+
 		/** Computes the mono-channel map. */
 		bool ComputeMonoChannelMap(Audio::FAlignedFloatBuffer& OutChannelMap);
 
diff --git a/Engine/Source/Runtime/Engine/Private/AudioDevice.cpp b/Engine/Source/Runtime/Engine/Private/AudioDevice.cpp
--- a/Engine/Source/Runtime/Engine/Private/AudioDevice.cpp
+++ b/Engine/Source/Runtime/Engine/Private/AudioDevice.cpp
@@ -697,6 +697,10 @@ void FAudioDevice::OnDeviceCreated(Audio::FDeviceId InDeviceID)
 	{
 		InitializeSubsystemCollection();
 		FAudioDeviceManagerDelegates::OnAudioDeviceCreated.Remove(DeviceCreatedHandle);
+
+#if UE_AUDIO_PROFILERTRACE_ENABLED
+		FTraceAuxiliary::OnTraceStarted.AddRaw(this, &FAudioDevice::OnTraceStarted);
+#endif
 	}
 }
 
@@ -705,6 +709,10 @@ void FAudioDevice::OnDeviceDestroyed(Audio::FDeviceId InDeviceID)
 	if (InDeviceID == DeviceID)
 	{
 		FAudioDeviceManagerDelegates::OnAudioDeviceDestroyed.Remove(DeviceDestroyedHandle);
+
+#if UE_AUDIO_PROFILERTRACE_ENABLED
+		FTraceAuxiliary::OnTraceStarted.RemoveAll(this);
+#endif
 	}
 }
 
@@ -2212,6 +2220,16 @@ bool FAudioDevice::HandleSetDynamicSoundCommand(const TCHAR* Cmd, FOutputDevice&
 }
 #endif // !UE_BUILD_SHIPPING
 
+#if UE_AUDIO_PROFILERTRACE_ENABLED
+void FAudioDevice::OnTraceStarted(FTraceAuxiliary::EConnectionType TraceType, const FString& TraceDestination)
+{
+	for (auto& [ActiveSound, VirtualLoop] : VirtualLoops)
+	{
+		VirtualLoop.OnTraceStarted();
+	}
+}
+#endif // UE_AUDIO_PROFILERTRACE_ENABLED
+
 bool FAudioDevice::IsHRTFEnabledForAll() const
 {
 	if (IsInAudioThread())
diff --git a/Engine/Source/Runtime/Engine/Private/AudioVirtualLoop.cpp b/Engine/Source/Runtime/Engine/Private/AudioVirtualLoop.cpp
--- a/Engine/Source/Runtime/Engine/Private/AudioVirtualLoop.cpp
+++ b/Engine/Source/Runtime/Engine/Private/AudioVirtualLoop.cpp
@@ -358,6 +358,30 @@ bool FAudioVirtualLoop::Update(float DeltaTime, bool bForceUpdate)
 	return true;
 }
 
+#if UE_AUDIO_PROFILERTRACE_ENABLED
+void FAudioVirtualLoop::OnTraceStarted()
+{
+	// Make sure we send mixer source data to Audio Insights if it's opened in the middle of a PIE session
+	const bool bChannelEnabled = UE_TRACE_CHANNELEXPR_IS_ENABLED(AudioMixerChannel);
+	if (bChannelEnabled && ActiveSound)
+	{
+		check(ActiveSound);
+		const USoundBase* Sound = ActiveSound->GetSound();
+		check(Sound);
+
+		if (const FAudioDevice* AudioDevice = ActiveSound->AudioDevice)
+		{
+			UE_TRACE_LOG(Audio, VirtualLoopVirtualize, AudioMixerChannel)
+				<< VirtualLoopVirtualize.DeviceId(static_cast<uint32>(AudioDevice->DeviceID))
+				<< VirtualLoopVirtualize.Timestamp(FPlatformTime::Cycles64())
+				<< VirtualLoopVirtualize.PlayOrder(ActiveSound->GetPlayOrder())
+				<< VirtualLoopVirtualize.ComponentId(ActiveSound->GetAudioComponentID())
+				<< VirtualLoopVirtualize.Name(Sound ? *Sound->GetPathName() : TEXT("N/A"));
+		}
+	}
+}
+#endif
+
 bool FAudioVirtualLoop::ShouldListenerMoveForceUpdate(const FTransform& LastTransform, const FTransform& CurrentTransform)
 {
 	const float DistanceSq = FVector::DistSquared(LastTransform.GetTranslation(), CurrentTransform.GetTranslation());
diff --git a/Engine/Source/Runtime/Engine/Public/AudioDevice.h b/Engine/Source/Runtime/Engine/Public/AudioDevice.h
--- a/Engine/Source/Runtime/Engine/Public/AudioDevice.h
+++ b/Engine/Source/Runtime/Engine/Public/AudioDevice.h
@@ -9,6 +9,7 @@
 #include "HAL/LowLevelMemStats.h"
 #include "Math/Transform.h"
 #include "Math/Vector.h"
+#include "ProfilingDebugging/TraceAuxiliary.h"
 #include "Sound/AudioVolume.h"
 #include "Sound/SoundClass.h"
 #include "Sound/SoundMix.h"
@@ -513,6 +514,9 @@ private:
 	void GetSoundClassInfo(TMap<FName, FAudioClassInfo>& AudioClassInfos);
 #endif
 
+#if UE_AUDIO_PROFILERTRACE_ENABLED
+	void OnTraceStarted(FTraceAuxiliary::EConnectionType TraceType, const FString& TraceDestination);
+#endif // UE_AUDIO_PROFILERTRACE_ENABLED
 
 public:
 
diff --git a/Engine/Source/Runtime/Engine/Public/AudioVirtualLoop.h b/Engine/Source/Runtime/Engine/Public/AudioVirtualLoop.h
--- a/Engine/Source/Runtime/Engine/Public/AudioVirtualLoop.h
+++ b/Engine/Source/Runtime/Engine/Public/AudioVirtualLoop.h
@@ -2,6 +2,7 @@
 #pragma once
 
 #include "ActiveSound.h"
+#include "AudioMixerTrace.h"
 #include "Sound/SoundModulationDestination.h"
 
 /**
@@ -83,4 +84,8 @@ public:
 	  * Returns whether or not the sound is ready to be realized.
 	  */
 	ENGINE_API bool Update(float DeltaTime, bool bForceUpdate);
+
+#if UE_AUDIO_PROFILERTRACE_ENABLED
+	ENGINE_API void OnTraceStarted();
+#endif
 };
