Speculative crash fix - null check ModSystem in FControlBusProxy::OnTraceStarted

diff --git a/Engine/Plugins/Runtime/AudioModulation/Source/AudioModulation/Private/SoundControlBusProxy.cpp b/Engine/Plugins/Runtime/AudioModulation/Source/AudioModulation/Private/SoundControlBusProxy.cpp
--- a/Engine/Plugins/Runtime/AudioModulation/Source/AudioModulation/Private/SoundControlBusProxy.cpp
+++ b/Engine/Plugins/Runtime/AudioModulation/Source/AudioModulation/Private/SoundControlBusProxy.cpp
@@ -196,6 +196,11 @@ namespace AudioModulation
 	void AudioModulation::FControlBusProxy::OnTraceStarted(FTraceAuxiliary::EConnectionType TraceType, const FString& TraceDestination)
 	{
 #if UE_AUDIO_PROFILERTRACE_ENABLED
+		if (ModSystem == nullptr)
+		{
+			return;
+		}
+
 		UE_TRACE_LOG(Audio, ControlBusActivate, AudioChannel)
 			<< ControlBusActivate.DeviceId(ModSystem->GetAudioDeviceId())
 			<< ControlBusActivate.ControlBusId(GetId())
