Pseudo-backout of a change causing editor hangs/crashes. Disabling the registration of the callbacks.

#rnx

diff --git a/Engine/Plugins/Runtime/AudioModulation/Source/AudioModulation/Private/SoundControlBusProxy.cpp b/Engine/Plugins/Runtime/AudioModulation/Source/AudioModulation/Private/SoundControlBusProxy.cpp
--- a/Engine/Plugins/Runtime/AudioModulation/Source/AudioModulation/Private/SoundControlBusProxy.cpp
+++ b/Engine/Plugins/Runtime/AudioModulation/Source/AudioModulation/Private/SoundControlBusProxy.cpp
@@ -254,15 +254,15 @@ namespace AudioModulation
 
 	void AudioModulation::FControlBusProxy::RegisterForTraceMessages()
 	{
-#if UE_AUDIO_PROFILERTRACE_ENABLED
+#if UE_AUDIO_PROFILERTRACE_ENABLED && 0
 		FTraceAuxiliary::OnTraceStarted.AddRaw(this, &FControlBusProxy::OnTraceStarted);
 #endif // UE_AUDIO_PROFILERTRACE_ENABLED
 	}
 
 	void AudioModulation::FControlBusProxy::UnregisterForTraceMessages()
 	{
-#if UE_AUDIO_PROFILERTRACE_ENABLED
-		FTraceAuxiliary::OnTraceStarted.RemoveAll(this);
+#if UE_AUDIO_PROFILERTRACE_ENABLED && 0
+		FTraceAuxiliary::OnTraceStarted.RemoveAll(this); 
 #endif // UE_AUDIO_PROFILERTRACE_ENABLED
 	}
 } // namespace AudioModulation
