Re-implement initialization of modulation window in Audio Insights without causing editor hangs due to delegate calls to dead objects.

- Some users were experiencing editor hangs due to SoundControlBusProxy objects not always unregistering their delegate handles to OnTraceStarted and the ModSystem pointer being null.
- Instead of handling registering/unregistering to the OnTraceStarted delegate within the SoundControlBusProxy objects, we now piggyback the registration to the same delegate within AudioModulationSystem, and pass the ModSystem down to SoundControlBusProxy object. The lifetime of the system object is much more predictable and should be a safer way to handle this behaviour.

diff --git a/Engine/Plugins/Runtime/AudioModulation/Source/AudioModulation/Private/AudioModulationSystem.cpp b/Engine/Plugins/Runtime/AudioModulation/Source/AudioModulation/Private/AudioModulationSystem.cpp
--- a/Engine/Plugins/Runtime/AudioModulation/Source/AudioModulation/Private/AudioModulationSystem.cpp
+++ b/Engine/Plugins/Runtime/AudioModulation/Source/AudioModulation/Private/AudioModulationSystem.cpp
@@ -848,6 +848,8 @@ namespace AudioModulation
 					<< BusMixRegisterBus.Timestamp(FPlatformTime::Cycles64())
 					<< BusMixRegisterBus.ModulatingSourceId(BusId)
 					<< BusMixRegisterBus.BusName(*(ControlBusProxy.GetName().ToString()));
+
+				ControlBusProxy.OnTraceStarted(*this);
 			}
 
 			for (const auto& [BusMixId, ModulatorBusMixProxy] : RefProxies.BusMixes)
@@ -865,12 +867,7 @@ namespace AudioModulation
 	Audio::FModulatorTypeId FAudioModulationSystem::RegisterModulator(Audio::FModulatorHandleId InHandleId, const FControlBusSettings& InSettings)
 	{
 		FControlBusSettings CachedSettings = InSettings;
-		RegisterModulator(InHandleId, MoveTemp(CachedSettings), RefProxies.Buses, RefModulators.BusMap, [this](FBusHandle& NewHandle)
-		{
-#if UE_AUDIO_PROFILERTRACE_ENABLED
-			NewHandle.FindProxy().RegisterForTraceMessages();
-#endif // #if UE_AUDIO_PROFILERTRACE_ENABLED
-		});
+		RegisterModulator(InHandleId, MoveTemp(CachedSettings), RefProxies.Buses, RefModulators.BusMap);
 
 		return static_cast<Audio::FModulatorTypeId>(EModulatorType::Bus);
 	}
@@ -975,7 +972,6 @@ namespace AudioModulation
 				if (RefModulators.BusMap.Find(BusHandle) == nullptr)
 				{
 					FControlBusProxy& ControlBusProxy = BusHandle.FindProxy();
-					ControlBusProxy.UnregisterForTraceMessages();
 
 					for (const FGeneratorHandle& GeneratorHandle : ControlBusProxy.GetGeneratorHandles())
 					{
diff --git a/Engine/Plugins/Runtime/AudioModulation/Source/AudioModulation/Private/SoundControlBusProxy.cpp b/Engine/Plugins/Runtime/AudioModulation/Source/AudioModulation/Private/SoundControlBusProxy.cpp
--- a/Engine/Plugins/Runtime/AudioModulation/Source/AudioModulation/Private/SoundControlBusProxy.cpp
+++ b/Engine/Plugins/Runtime/AudioModulation/Source/AudioModulation/Private/SoundControlBusProxy.cpp
@@ -193,16 +193,11 @@ namespace AudioModulation
 		return OutValue;
 	}
 
-	void AudioModulation::FControlBusProxy::OnTraceStarted(FTraceAuxiliary::EConnectionType TraceType, const FString& TraceDestination)
-	{
 #if UE_AUDIO_PROFILERTRACE_ENABLED
-		if (ModSystem == nullptr)
-		{
-			return;
-		}
-
+	void AudioModulation::FControlBusProxy::OnTraceStarted(FAudioModulationSystem& InModSystem) const
+	{
 		UE_TRACE_LOG(Audio, ControlBusActivate, AudioChannel)
-			<< ControlBusActivate.DeviceId(ModSystem->GetAudioDeviceId())
+			<< ControlBusActivate.DeviceId(InModSystem.GetAudioDeviceId())
 			<< ControlBusActivate.ControlBusId(GetId())
 			<< ControlBusActivate.Timestamp(FPlatformTime::Cycles64())
 			<< ControlBusActivate.Name(*(GetName().ToString()))
@@ -211,20 +206,20 @@ namespace AudioModulation
 		for (const FGeneratorHandle& GeneratorHandle: GeneratorHandles)
 		{
 			UE_TRACE_LOG(Audio, GeneratorRegisterBus, AudioChannel)
-				<< GeneratorRegisterBus.DeviceId(ModSystem->GetAudioDeviceId())
+				<< GeneratorRegisterBus.DeviceId(InModSystem.GetAudioDeviceId())
 				<< GeneratorRegisterBus.SourceId(GetId())
 				<< GeneratorRegisterBus.Timestamp(FPlatformTime::Cycles64())
 				<< GeneratorRegisterBus.ModulatingSourceId(GeneratorHandle.GetId())
 				<< GeneratorRegisterBus.BusName(*(GetName().ToString()));
 
 			UE_TRACE_LOG(Audio, GeneratorActivate, AudioChannel)
-				<< GeneratorActivate.DeviceId(ModSystem->GetAudioDeviceId())
+				<< GeneratorActivate.DeviceId(InModSystem.GetAudioDeviceId())
 				<< GeneratorActivate.SourceId(GeneratorHandle.GetId())
 				<< GeneratorActivate.Timestamp(FPlatformTime::Cycles64())
 				<< GeneratorActivate.Name(*(GeneratorHandle.FindProxy().GetName().ToString()));
 		}
-#endif // UE_AUDIO_PROFILERTRACE_ENABLED
 	}
+#endif // UE_AUDIO_PROFILERTRACE_ENABLED
 
 	void FControlBusProxy::MixIn(const float InValue)
 	{
@@ -251,18 +246,4 @@ namespace AudioModulation
 		GeneratorValue = 1.0f;
 		MixValue = NAN;
 	}
-
-	void AudioModulation::FControlBusProxy::RegisterForTraceMessages()
-	{
-#if UE_AUDIO_PROFILERTRACE_ENABLED && 0
-		FTraceAuxiliary::OnTraceStarted.AddRaw(this, &FControlBusProxy::OnTraceStarted);
-#endif // UE_AUDIO_PROFILERTRACE_ENABLED
-	}
-
-	void AudioModulation::FControlBusProxy::UnregisterForTraceMessages()
-	{
-#if UE_AUDIO_PROFILERTRACE_ENABLED && 0
-		FTraceAuxiliary::OnTraceStarted.RemoveAll(this); 
-#endif // UE_AUDIO_PROFILERTRACE_ENABLED
-	}
 } // namespace AudioModulation
diff --git a/Engine/Plugins/Runtime/AudioModulation/Source/AudioModulation/Private/SoundControlBusProxy.h b/Engine/Plugins/Runtime/AudioModulation/Source/AudioModulation/Private/SoundControlBusProxy.h
--- a/Engine/Plugins/Runtime/AudioModulation/Source/AudioModulation/Private/SoundControlBusProxy.h
+++ b/Engine/Plugins/Runtime/AudioModulation/Source/AudioModulation/Private/SoundControlBusProxy.h
@@ -83,15 +83,14 @@ namespace AudioModulation
 		void MixGenerators();
 		void Reset();
 
-		void RegisterForTraceMessages();
-		void UnregisterForTraceMessages();
+#if UE_AUDIO_PROFILERTRACE_ENABLED
+		void OnTraceStarted(FAudioModulationSystem& InModSystem) const;
+#endif // UE_AUDIO_PROFILERTRACE_ENABLED
 
 	private:
 		void Init(FControlBusSettings&& InSettings);
 		float Mix(float ValueA) const;
 
-		void OnTraceStarted(FTraceAuxiliary::EConnectionType TraceType, const FString& TraceDestination);
-
 		float DefaultValue;
 
 		// Cached values
