Control Bus Mix Timers & Retrigger Option

diff --git a/Engine/Plugins/Runtime/AudioModulation/Source/AudioModulation/Private/AudioModulation.cpp b/Engine/Plugins/Runtime/AudioModulation/Source/AudioModulation/Private/AudioModulation.cpp
--- a/Engine/Plugins/Runtime/AudioModulation/Source/AudioModulation/Private/AudioModulation.cpp
+++ b/Engine/Plugins/Runtime/AudioModulation/Source/AudioModulation/Private/AudioModulation.cpp
@@ -160,9 +160,9 @@ namespace AudioModulation
 		return ModSystem->LoadMixFromProfile(InProfileIndex, OutBusMix);
 	}
 
-	void FAudioModulationManager::UpdateMix(const TArray<FSoundControlBusMixStage>& InStages, USoundControlBusMix& InOutMix, bool bInUpdateObject, float InFadeTime)
+	void FAudioModulationManager::UpdateMix(const TArray<FSoundControlBusMixStage>& InStages, USoundControlBusMix& InOutMix, bool bInUpdateObject, float InFadeTime, double Duration, bool bRetriggerOnActivation)
 	{
-		ModSystem->UpdateMix(InStages, InOutMix, bInUpdateObject, InFadeTime);
+		ModSystem->UpdateMix(InStages, InOutMix, bInUpdateObject, InFadeTime, Duration, bRetriggerOnActivation);
 	}
 
 	void FAudioModulationManager::UpdateMix(const USoundControlBusMix& InMix, float InFadeTime)
diff --git a/Engine/Plugins/Runtime/AudioModulation/Source/AudioModulation/Private/AudioModulationStatics.cpp b/Engine/Plugins/Runtime/AudioModulation/Source/AudioModulation/Private/AudioModulationStatics.cpp
--- a/Engine/Plugins/Runtime/AudioModulation/Source/AudioModulation/Private/AudioModulationStatics.cpp
+++ b/Engine/Plugins/Runtime/AudioModulation/Source/AudioModulation/Private/AudioModulationStatics.cpp
@@ -230,7 +230,7 @@ FSoundControlBusMixStage UAudioModulationStatics::CreateBusMixStage(const UObjec
 	return MixStage;
 }
 
-USoundControlBusMix* UAudioModulationStatics::CreateBusMix(UObject* WorldContextObject, FName Name, TArray<FSoundControlBusMixStage> Stages, bool Activate)
+USoundControlBusMix* UAudioModulationStatics::CreateBusMix(UObject* WorldContextObject, FName Name, TArray<FSoundControlBusMixStage> Stages, bool Activate, double Duration, bool bRetriggerOnActivation)
 {
 	UWorld* World = GetAudioWorld(WorldContextObject);
 	if (!World)
@@ -253,6 +253,9 @@ USoundControlBusMix* UAudioModulationStatics::CreateBusMix(UObject* WorldContext
 				*Name.ToString());
 		}
 	}
+	
+	NewBusMix->Duration = Duration;
+	NewBusMix->bRetriggerOnActivation = bRetriggerOnActivation;
 
 	if (Activate)
 	{
@@ -443,7 +446,7 @@ TArray<FSoundControlBusMixStage> UAudioModulationStatics::LoadMixFromProfile(con
 	return TArray<FSoundControlBusMixStage>();
 }
 
-void UAudioModulationStatics::UpdateMix(const UObject* WorldContextObject, USoundControlBusMix* Mix, TArray<FSoundControlBusMixStage> Stages, float InFadeTime)
+void UAudioModulationStatics::UpdateMix(const UObject* WorldContextObject, USoundControlBusMix* Mix, TArray<FSoundControlBusMixStage> Stages, float InFadeTime, double Duration, bool bRetriggerOnActivation)
 {
 	if (Mix)
 	{
@@ -452,7 +455,7 @@ void UAudioModulationStatics::UpdateMix(const UObject* WorldContextObject, USoun
 		{
 			// UObject representation is not updated in this form of the call as doing so from
 			// PIE can result in an unstable state where UObject is modified but not properly dirtied.
-			ModSystem->UpdateMix(Stages, *Mix, false /* bUpdateObject */, InFadeTime);
+			ModSystem->UpdateMix(Stages, *Mix, false /* bUpdateObject */, InFadeTime, Duration, bRetriggerOnActivation);
 		}
 	}
 }
diff --git a/Engine/Plugins/Runtime/AudioModulation/Source/AudioModulation/Private/AudioModulationSystem.cpp b/Engine/Plugins/Runtime/AudioModulation/Source/AudioModulation/Private/AudioModulationSystem.cpp
--- a/Engine/Plugins/Runtime/AudioModulation/Source/AudioModulation/Private/AudioModulationSystem.cpp
+++ b/Engine/Plugins/Runtime/AudioModulation/Source/AudioModulation/Private/AudioModulationSystem.cpp
@@ -166,7 +166,7 @@ namespace AudioModulation
 			FBusMixHandle BusMixHandle = FBusMixHandle::Get(Settings.GetId(), RefProxies.BusMixes);
 			if (BusMixHandle.IsValid())
 			{
-				BusMixHandle.FindProxy().SetEnabled(MoveTemp(Settings));
+				BusMixHandle.FindProxy().SetEnabled();
 			}
 			else
 			{
@@ -1000,7 +1000,7 @@ namespace AudioModulation
 		});
 	}
 
-	void FAudioModulationSystem::UpdateMix(const TArray<FSoundControlBusMixStage>& InStages, USoundControlBusMix& InOutMix, bool bInUpdateObject, float InFadeTime)
+	void FAudioModulationSystem::UpdateMix(const TArray<FSoundControlBusMixStage>& InStages, USoundControlBusMix& InOutMix, bool bInUpdateObject, float InFadeTime, double Duration, bool bRetriggerOnActivation)
 	{
 		if (bInUpdateObject)
 		{
@@ -1042,11 +1042,15 @@ namespace AudioModulation
 		}
 	
 		const FString BusMixName = InOutMix.GetName();
-		RunCommandOnProcessingThread([this, MixId, StageSettings, InFadeTime, BusMixName]()
+		RunCommandOnProcessingThread([this, MixId, StageSettings, InFadeTime, BusMixName, Duration, bRetriggerOnActivation]()
 		{
 			if (FModulatorBusMixProxy* BusMixes = RefProxies.BusMixes.Find(MixId))
 			{
-				BusMixes->SetMix(StageSettings, InFadeTime, BusMixName);
+				BusMixes->SetMixData(StageSettings, InFadeTime, BusMixName, Duration, bRetriggerOnActivation);
+			}
+			else
+			{
+				UE_LOG(LogAudioModulation, Display, TEXT("Could not update mix '%s' because the mix is not currently active."), *BusMixName);
 			}
 		});
 	}
diff --git a/Engine/Plugins/Runtime/AudioModulation/Source/AudioModulation/Private/AudioModulationSystem.h b/Engine/Plugins/Runtime/AudioModulation/Source/AudioModulation/Private/AudioModulationSystem.h
--- a/Engine/Plugins/Runtime/AudioModulation/Source/AudioModulation/Private/AudioModulationSystem.h
+++ b/Engine/Plugins/Runtime/AudioModulation/Source/AudioModulation/Private/AudioModulationSystem.h
@@ -136,7 +136,7 @@ namespace AudioModulation
 		 * Updates mix/mix by filter, modifying the mix instance if it is active. If bInUpdateObject is true,
 		 * updates UObject definition in addition to proxy.
 		 */
-		void UpdateMix(const TArray<FSoundControlBusMixStage>& InStages, USoundControlBusMix& InOutMix, bool bInUpdateObject = false, float InFadeTime = -1.0f);
+		void UpdateMix(const TArray<FSoundControlBusMixStage>& InStages, USoundControlBusMix& InOutMix, bool bInUpdateObject = false, float InFadeTime = -1.0f, double Duration = -1.0, bool bRetriggerOnActivation = false);
 		void UpdateMixByFilter(const FString& InAddressFilter, const TSubclassOf<USoundModulationParameter>& InParamClassFilter, USoundModulationParameter* InParamFilter, float Value, float FadeTime, USoundControlBusMix& InOutMix, bool bInUpdateObject = false);
 
 		/*
diff --git a/Engine/Plugins/Runtime/AudioModulation/Source/AudioModulation/Private/SoundControlBusMix.cpp b/Engine/Plugins/Runtime/AudioModulation/Source/AudioModulation/Private/SoundControlBusMix.cpp
--- a/Engine/Plugins/Runtime/AudioModulation/Source/AudioModulation/Private/SoundControlBusMix.cpp
+++ b/Engine/Plugins/Runtime/AudioModulation/Source/AudioModulation/Private/SoundControlBusMix.cpp
@@ -33,7 +33,6 @@ FSoundControlBusMixStage::FSoundControlBusMixStage(USoundControlBus* InBus, cons
 
 USoundControlBusMix::USoundControlBusMix(const FObjectInitializer& ObjectInitializer)
 	: Super(ObjectInitializer)
-	, ProfileIndex(0)
 {
 }
 
diff --git a/Engine/Plugins/Runtime/AudioModulation/Source/AudioModulation/Private/SoundControlBusMixProxy.cpp b/Engine/Plugins/Runtime/AudioModulation/Source/AudioModulation/Private/SoundControlBusMixProxy.cpp
--- a/Engine/Plugins/Runtime/AudioModulation/Source/AudioModulation/Private/SoundControlBusMixProxy.cpp
+++ b/Engine/Plugins/Runtime/AudioModulation/Source/AudioModulation/Private/SoundControlBusMixProxy.cpp
@@ -50,7 +50,10 @@ namespace AudioModulation
 	}
 
 	FModulatorBusMixSettings::FModulatorBusMixSettings(const USoundControlBusMix& InBusMix)
-		: TModulatorBase<FBusMixId>(InBusMix.GetFName(), InBusMix.GetUniqueID())
+		: TModulatorBase<FBusMixId>(InBusMix.GetFName()
+		, InBusMix.GetUniqueID())
+		, Duration(InBusMix.Duration)
+		, bRetriggerOnActivation(InBusMix.bRetriggerOnActivation)
 	{
 		for (const FSoundControlBusMixStage& Stage : InBusMix.MixStages)
 		{
@@ -71,18 +74,23 @@ namespace AudioModulation
 	FModulatorBusMixSettings::FModulatorBusMixSettings(FModulatorBusMixSettings&& InBusMixSettings)
 		: TModulatorBase<FBusMixId>(InBusMixSettings.GetName(), InBusMixSettings.GetId())
 		, Stages(MoveTemp(InBusMixSettings.Stages))
+		, Duration(InBusMixSettings.Duration)
+		, bRetriggerOnActivation(InBusMixSettings.bRetriggerOnActivation)
 	{
 	}
 
 	FModulatorBusMixProxy::FModulatorBusMixProxy(FModulatorBusMixSettings&& InSettings, FAudioModulationSystem& OutModSystem)
 		: TModulatorProxyRefType(InSettings.GetName(), InSettings.GetId(), OutModSystem)
+		, Duration(InSettings.Duration)
+		, TimeRemaining(Duration)
+		, bRetriggerOnActivation(InSettings.bRetriggerOnActivation)
 	{
-		SetEnabled(MoveTemp(InSettings));
+		SetMixDataAndEnable(MoveTemp(InSettings));
 	}
 
 	FModulatorBusMixProxy& FModulatorBusMixProxy::operator=(FModulatorBusMixSettings&& InSettings)
 	{
-		SetEnabled(MoveTemp(InSettings));
+		SetMixDataAndEnable(MoveTemp(InSettings));
 
 		return *this;
 	}
@@ -97,30 +105,75 @@ namespace AudioModulation
 		Stages.Reset();
 	}
 
-	void FModulatorBusMixProxy::SetEnabled(FModulatorBusMixSettings&& InSettings)
+	void FModulatorBusMixProxy::SetMixDataAndEnable(FModulatorBusMixSettings&& InSettings)
+	{
+		SetMixData(MoveTemp(InSettings));
+		SetEnabled();
+	}
+
+	void FModulatorBusMixProxy::SetEnabled()
 	{
 		check(ModSystem);
 
+		Status = EStatus::Enabled;
+		
+		if (Duration >= 0.0f)
+		{
+			TimeRemaining = FMath::Max(Duration, 0);
+		}
+		
+		if (bRetriggerOnActivation)
+		{
+			int StageIndex = 0;
+			for (TPair<FBusId, FModulatorBusMixStageProxy>& IdProxyPair : Stages)
+			{
+				FModulatorBusMixStageProxy& Stage = IdProxyPair.Value;
+				const float DefaultValue = Stage.BusHandle.FindProxy().GetDefaultValue();
+				Stage.Value.SetCurrentValue(DefaultValue);
+				Stage.Value.SetActiveFade(FSoundModulationMixValue::EActiveFade::Attack);
+				Stage.Value.TargetValue = StageValues[StageIndex];
+				++StageIndex;
+			}
+		}
+	}
+
+	void FModulatorBusMixProxy::SetMixData(FModulatorBusMixSettings&& InSettings)
+	{
+		check(ModSystem);
+
+		Duration = InSettings.Duration;
+		bRetriggerOnActivation = InSettings.bRetriggerOnActivation;
+		
 		// Cache stages to avoid releasing stage state (and potentially referenced bus state) when re-enabling
 		FStageMap CachedStages = Stages;
-
-		Status = EStatus::Enabled;
 		Stages.Reset();
+		StageValues.Reset();
 		for (FModulatorBusMixStageSettings& StageSettings : InSettings.Stages)
 		{
 			const FBusId BusId = StageSettings.GetId();
 			FModulatorBusMixStageProxy StageProxy(MoveTemp(StageSettings), *ModSystem);
 			if (const FModulatorBusMixStageProxy* CachedStage = CachedStages.Find(BusId))
 			{
+				const float StageValue = CachedStage->Value.GetCurrentValue();
 				StageProxy.Value.SetCurrentValue(CachedStage->Value.GetCurrentValue());
 			}
 
+			StageValues.Add(StageSettings.Value.TargetValue);
 			Stages.Add(BusId, MoveTemp(StageProxy));
 		}
 	}
 
-	void FModulatorBusMixProxy::SetMix(const TArray<FModulatorBusMixStageSettings>& InStages, float InFadeTime, const FString& BusMixName)
+	void FModulatorBusMixProxy::SetMixData(const TArray<FModulatorBusMixStageSettings>& InStages, float InFadeTime, const FString& BusMixName, double InDuration, bool bInRetriggerOnActivation)
 	{
+		
+		if (InDuration >= 0.0f)
+		{
+			Duration = InDuration;
+			TimeRemaining = FMath::Max(Duration, 0);
+		}
+		
+		bRetriggerOnActivation = bInRetriggerOnActivation;
+		
 		for (const FModulatorBusMixStageSettings& NewStage : InStages)
 		{
 			const FBusId BusId = NewStage.GetId();
@@ -175,6 +228,16 @@ namespace AudioModulation
 
 	void FModulatorBusMixProxy::Update(const double InElapsed, FBusProxyMap& OutProxyMap)
 	{
+		if (Status == EStatus::Enabled && Duration >= 0)
+		{
+			TimeRemaining -= InElapsed;
+			if (TimeRemaining <= 0.0)
+			{
+				UE_LOG(LogAudioModulation, Display, TEXT("Automatically deactivating mix after %.3f seconds"), FMath::Max(Duration, 0));
+				SetStopping();
+			}
+		}
+		
 		bool bRequestStop = true;
 		for (TPair<FBusId, FModulatorBusMixStageProxy>& Stage : Stages)
 		{
diff --git a/Engine/Plugins/Runtime/AudioModulation/Source/AudioModulation/Private/SoundControlBusMixProxy.h b/Engine/Plugins/Runtime/AudioModulation/Source/AudioModulation/Private/SoundControlBusMixProxy.h
--- a/Engine/Plugins/Runtime/AudioModulation/Source/AudioModulation/Private/SoundControlBusMixProxy.h
+++ b/Engine/Plugins/Runtime/AudioModulation/Source/AudioModulation/Private/SoundControlBusMixProxy.h
@@ -40,6 +40,9 @@ namespace AudioModulation
 		FModulatorBusMixSettings(FModulatorBusMixSettings&& InSettings);
 
 		TArray<FModulatorBusMixStageSettings> Stages;
+
+		double Duration = -1.0;
+		bool bRetriggerOnActivation = false;
 	};
 
 	class FModulatorBusMixStageProxy : public TModulatorBase<FBusId>
@@ -74,10 +77,12 @@ namespace AudioModulation
 		// Resets stage map
 		void Reset();
 
-		void SetEnabled(FModulatorBusMixSettings&& InSettings);
-		void SetMix(const TArray<FModulatorBusMixStageSettings>& InStages, float InFadeTime, const FString& BusMixName);
-		void SetMixByFilter(const FString& InAddressFilter, uint32 InParamClassId, uint32 InParamId, float InValue, float InFadeTime);
+		void SetEnabled();
 		void SetStopping();
+		// Set stage and Mix Config settings. This does not enable/activate the mix; for that, use SetEnabled
+		void SetMixData(const TArray<FModulatorBusMixStageSettings>& InStages, float InFadeTime, const FString& BusMixName, double InDuration, bool bInRetriggerOnActivation);
+		// Set the stage values of all buses which satisfy the input filter data. If any of the filters (Address, Parameter type, or specific parameter) are satisfied, that stage's value will be set 
+		void SetMixByFilter(const FString& InAddressFilter, uint32 InParamClassId, uint32 InParamId, float InValue, float InFadeTime);
 
 		void Update(const double Elapsed, FBusProxyMap& ProxyMap);
 
@@ -85,7 +90,15 @@ namespace AudioModulation
 		FStageMap Stages;
 
 	private:
-		EStatus Status;
+		EStatus Status = EStatus::Stopped;
+		double Duration = -1.0;
+		double TimeRemaining = 0.0;
+		bool bRetriggerOnActivation = false;
+
+		TArray<float> StageValues;
+
+		void SetMixDataAndEnable(FModulatorBusMixSettings&& InSettings);
+		void SetMixData(FModulatorBusMixSettings&& InSettings);
 	};
 
 	using FBusMixProxyMap = TMap<FBusMixId, FModulatorBusMixProxy>;
diff --git a/Engine/Plugins/Runtime/AudioModulation/Source/AudioModulation/Private/SoundModulationValue.cpp b/Engine/Plugins/Runtime/AudioModulation/Source/AudioModulation/Private/SoundModulationValue.cpp
--- a/Engine/Plugins/Runtime/AudioModulation/Source/AudioModulation/Private/SoundModulationValue.cpp
+++ b/Engine/Plugins/Runtime/AudioModulation/Source/AudioModulation/Private/SoundModulationValue.cpp
@@ -17,13 +17,6 @@ FSoundModulationMixValue::FSoundModulationMixValue(float InValue, float InAttack
 
 void FSoundModulationMixValue::SetActiveFade(EActiveFade InActiveFade, float InFadeTime)
 {
-	// Release has already been issued, so fade cannot be overwritten (to avoid orphaned
-	// mix stages/values)
-	if (ActiveFade == EActiveFade::Release)
-	{
-		return;
-	}
-
 	// Only override fades can modify existing, active behavior
 	if (ActiveFade != EActiveFade::Override)
 	{
diff --git a/Engine/Plugins/Runtime/AudioModulation/Source/AudioModulation/Public/AudioModulation.h b/Engine/Plugins/Runtime/AudioModulation/Source/AudioModulation/Public/AudioModulation.h
--- a/Engine/Plugins/Runtime/AudioModulation/Source/AudioModulation/Public/AudioModulation.h
+++ b/Engine/Plugins/Runtime/AudioModulation/Source/AudioModulation/Public/AudioModulation.h
@@ -85,7 +85,7 @@ namespace AudioModulation
 		UE_API void SaveMixToProfile(const USoundControlBusMix& InBusMix, const int32 InProfileIndex);
 		UE_API TArray<FSoundControlBusMixStage> LoadMixFromProfile(const int32 InProfileIndex, USoundControlBusMix& OutBusMix);
 
-		UE_API void UpdateMix(const TArray<FSoundControlBusMixStage>& InStages, USoundControlBusMix& InOutMix, bool bInUpdateObject = false, float InFadeTime = -1.0f);
+		UE_API void UpdateMix(const TArray<FSoundControlBusMixStage>& InStages, USoundControlBusMix& InOutMix, bool bInUpdateObject = false, float InFadeTime = -1.0f, double Duration = -1.0, bool bRetriggerOnActivation = false);
 		UE_API void UpdateMix(const USoundControlBusMix& InMix, float InFadeTime = -1.0f);
 		UE_API void UpdateMixByFilter(const FString& InAddressFilter, const TSubclassOf<USoundModulationParameter>& InParamClassFilter, USoundModulationParameter* InParamFilter, float Value, float FadeTime, USoundControlBusMix& InOutMix, bool bInUpdateObject = false);
 
diff --git a/Engine/Plugins/Runtime/AudioModulation/Source/AudioModulation/Public/AudioModulationStatics.h b/Engine/Plugins/Runtime/AudioModulation/Source/AudioModulation/Public/AudioModulationStatics.h
--- a/Engine/Plugins/Runtime/AudioModulation/Source/AudioModulation/Public/AudioModulationStatics.h
+++ b/Engine/Plugins/Runtime/AudioModulation/Source/AudioModulation/Public/AudioModulationStatics.h
@@ -61,12 +61,12 @@ public:
 		WorldContext = "WorldContextObject",
 		Keywords = "modulation modulator")
 	)
-	static UE_API void ActivateBusMix(const UObject* WorldContextObject, USoundControlBusMix* Mix);
+	static UE_API void ActivateBusMix(const UObject* WorldContextObject, USoundControlBusMix* BusMix);
 
 	/** Manually activates a modulation generator. If called, deactivation will only occur
 	 * if generator is manually deactivated and not referenced or destroyed (i.e. will not deactivate
 	 * when all references become inactive).
-	 * @param Modulator - Modulator to activate
+	 * @param Generator - Generator to activate
 	 */
 	UE_DEPRECATED(5.6, "To manually control the activation state of a modulator, use CreateModulationDestination.")
 	UFUNCTION(BlueprintCallable, Category = "Audio|Modulation", DisplayName = "Activate Modulation Generator", meta = (
@@ -90,7 +90,7 @@ public:
 	)
 	static UE_API UPARAM(DisplayName = "Bus") USoundControlBus* CreateBus(UObject* WorldContextObject, FName Name, USoundModulationParameter* Parameter, bool Activate = false);
 
-	/* Create a mix with stages created for each provided bus that are initialized to the supplied value and timing parameters.
+	/** Create a mix with stages created for each provided bus that are initialized to the supplied value and timing parameters.
 	 * @param Buses - Buses to assign stages within new mix to
 	 * @param Value - Initial value for all stages created within the new mix.
 	 * @param AttackTime - Fade time to user when mix activates.
@@ -128,23 +128,30 @@ public:
 		float AttackTime = 0.1f,
 		float ReleaseTime = 0.1f);
 
-	/** Creates a modulation bus mix, with a bus stage set to the provided target value.
+	/** Creates a Control Bus Mix with the given stage configuration.
 	 * @param Name - Name of mix.
 	 * @param Stages - Stages mix is responsible for.
 	 * @param Activate - Whether or not to activate mix on creation. If true, deactivation will only occur
 	 * if returned mix is manually deactivated and not referenced or destroyed (i.e. will not deactivate
-	 * when all references become inactive).
+	 * when all references become inactive). It will also still deactivate based on the Duration value.
+	 * @param Duration - Amount of time the Mix will stay active for before deactivating itself.
+	 *					 When less than 0, the mix will stay active until manually deactivated.
+	 * @param bRetriggerOnActivation - If true, if this mix is already active when Activate is called,
+	 *								   stages will return to their default values before activating.
 	 * @return Capture this in a Blueprint variable to prevent it from being garbage collected. 
 	 */
 	UFUNCTION(BlueprintCallable, Category = "Audio|Modulation", DisplayName = "Create Control Bus Mix", meta = (
 		WorldContext = "WorldContextObject",
+		AdvancedDisplay = "3",
 		Keywords = "make modulation modulator")
 	)
 	static UE_API UPARAM(DisplayName = "BusMix") USoundControlBusMix* CreateBusMix(
 		UObject* WorldContextObject,
 		FName Name, 
 		TArray<FSoundControlBusMixStage> Stages,
-		bool Activate);
+		bool Activate,
+		double Duration = -1.0,
+		bool bRetriggerOnActivation = false);
 
 	/** Creates a modulation parameter of a given class.
 	 * @param Name - Name of parameter.
@@ -164,7 +171,7 @@ public:
 
 	/** Creates a modulation generator based on an Envelope Follower with the given parameters.
 	 * @param Name - Name of generator.
-	 * @param Settings - The Envelope Follower settings, including what data to follow.
+	 * @param Params - The Envelope Follower settings, including what data to follow.
 	 * @return Capture this in a Blueprint variable to prevent it from being garbage collected.
 	 */
 	UFUNCTION(BlueprintCallable, Category = "Audio|Modulation", DisplayName = "Create Envelope Follower Generator", meta = (
@@ -178,7 +185,7 @@ public:
 
 	/** Creates a modulation generator based on an LFO with the given parameters.
 	 * @param Name - Name of generator.
-	 * @param Settings - The LFO Settings.
+	 * @param Params - The LFO Settings.
 	 * @return Capture this in a Blueprint variable to prevent it from being garbage collected.
 	 */
 	UFUNCTION(BlueprintCallable, Category = "Audio|Modulation", DisplayName = "Create LFO Generator", meta = (
@@ -192,7 +199,7 @@ public:
 
 	/** Creates a modulation generator based on an Attack/Decay Envelope.
 	 * @param Name - Name of generator.
-	 * @param Settings - The AD Envelope Settings.
+	 * @param Params - The AD Envelope Settings.
 	 * @return Capture this in a Blueprint variable to prevent it from being garbage collected.
 	 */
 	UFUNCTION(BlueprintCallable, Category = "Audio|Modulation", DisplayName = "Create AD Envelope Generator", meta = (
@@ -225,7 +232,7 @@ public:
 	static UE_API void DeactivateBus(const UObject* WorldContextObject, USoundControlBus* Bus);
 
 	/** Deactivates a modulation bus mix. Does nothing if an instance of the provided bus mix is already inactive.
-	 * @param BusMix - Mix to deactivate
+	 * @param Mix - Mix to deactivate
 	 */
 	UFUNCTION(BlueprintCallable, Category = "Audio|Modulation", DisplayName = "Deactivate Control Bus Mix", meta = (
 		WorldContext = "WorldContextObject",
@@ -235,7 +242,6 @@ public:
 
 	/** Deactivates a modulation generator. Does nothing if an instance of the provided generator is already inactive.
 	 * @param Generator - Generator to activate
-	 * @param Scope - Scope of modulator
 	 */
 	UE_DEPRECATED(5.6, "To manually control the activation state of a modulator, use CreateModulationDestination.")
 	UFUNCTION(BlueprintCallable, Category = "Audio|Modulation", DisplayName = "Deactivate Modulation Generator (Deprecated - 5.4)", meta = (
@@ -256,7 +262,7 @@ public:
 
 	/** Saves control bus mix to a profile, serialized to an ini file.  If mix is loaded, uses current proxy's state.
 	 * If not, uses default UObject representation.
-	 * @param BusMix - Mix object to serialize to profile .ini.
+	 * @param Mix - Mix object to serialize to profile .ini.
 	 * @param ProfileIndex - Index of profile, allowing multiple profiles can be saved for single mix object. If 0, saves to default ini profile (no suffix).
 	 */
 	UFUNCTION(BlueprintCallable, Category = "Audio|Modulation", DisplayName = "Save Control Bus Mix to Profile", meta = (
@@ -267,7 +273,7 @@ public:
 	static UE_API void SaveMixToProfile(const UObject* WorldContextObject, USoundControlBusMix* Mix, int32 ProfileIndex = 0);
 
 	/** Loads control bus mix from a profile into UObject mix definition, deserialized from an ini file.
-	 * @param BusMix - Mix object to deserialize profile .ini to.
+	 * @param Mix - Mix object to deserialize profile .ini to.
 	 * @param bActivate - If true, activate mix upon loading from profile.
 	 * @param ProfileIndex - Index of profile, allowing multiple profiles to be loaded to single mix object. If <= 0, loads from default profile (no suffix).
 	 * @return Stages - Stage values loaded from profile (empty if profile did not exist or had no values serialized).
@@ -287,12 +293,18 @@ public:
 	 * @param FadeTime - Fade time to user when interpolating between current value and new values.
 	 *					 If negative, falls back to last fade time set on stage. If fade time never set on stage,
 	 *					 uses attack time set on stage in mix asset.
+	 * @param Duration - Amount of time the Mix is activated for. When the Mix has been active for the given time,
+	 *					 automatically deactivates itself. When less than 0, the duration is infinite
+	 *					 (i.e. mix will stay active until manually deactivated).
+	 * @param bRetriggerOnActivation - If true, if this mix is already active when Activate is called,
+	 *								  stages will return to their default values before activating.
 	 */
 	UFUNCTION(BlueprintCallable, Category = "Audio|Modulation", DisplayName = "Set Control Bus Mix", meta = (
 		WorldContext = "WorldContextObject",
+		AdvancedDisplay = "3",
 		Keywords = "modulation modulator stage")
 	)
-	static UE_API void UpdateMix(const UObject* WorldContextObject, USoundControlBusMix* Mix, TArray<FSoundControlBusMixStage> Stages, float FadeTime = -1.0f);
+	static void UpdateMix(const UObject* WorldContextObject, USoundControlBusMix* Mix, TArray<FSoundControlBusMixStage> Stages, float FadeTime = -1.0f, double Duration = -1.0, bool bRetriggerOnActivation = false);
 
 	/** Sets a Global Control Bus Mix with a single stage associated with the provided Bus to the given float value.  This call should
 	 * be reserved for buses that are to be always active. It is *NOT* recommended for transient buses, as not calling clear can keep
diff --git a/Engine/Plugins/Runtime/AudioModulation/Source/AudioModulation/Public/SoundControlBusMix.h b/Engine/Plugins/Runtime/AudioModulation/Source/AudioModulation/Public/SoundControlBusMix.h
--- a/Engine/Plugins/Runtime/AudioModulation/Source/AudioModulation/Public/SoundControlBusMix.h
+++ b/Engine/Plugins/Runtime/AudioModulation/Source/AudioModulation/Public/SoundControlBusMix.h
@@ -33,42 +33,57 @@ struct FSoundControlBusMixStage
 	FSoundModulationMixValue Value;
 };
 
-UCLASS(config = Engine, autoexpandcategories = (Stage, Mix), editinlinenew, BlueprintType, MinimalAPI)
+UCLASS(config = Engine, autoexpandcategories = (Stage, Mix), editinlinenew, BlueprintType, MinimalAPI, meta = (PrioritizeCategories = "Mix Config Stages"))
 class USoundControlBusMix : public UObject
 {
 	GENERATED_UCLASS_BODY()
 
 protected:
 	// Loads the mix from the provided profile index
-	UFUNCTION(Category = Mix, meta = (CallInEditor = "true"))
+	UFUNCTION(Category = Mix, meta = (CallInEditor = "true", DisplayPriority = "50"))
 	void LoadMixFromProfile();
 
 	// Saves the mix to the provided profile index
-	UFUNCTION(Category = Mix, meta = (CallInEditor = "true"))
+	UFUNCTION(Category = Mix, meta = (CallInEditor = "true", DisplayPriority = "60"))
 	void SaveMixToProfile();
 
 	// Solos this mix, deactivating all others and activating this
 	// (if its not already active), while testing in-editor in all
 	// active worlds
-	UFUNCTION(Category = Mix, meta = (CallInEditor = "true"))
+	UFUNCTION(Category = Mix, meta = (CallInEditor = "true", DisplayPriority = "40"))
 	void SoloMix();
 
-	// Activates this mix in all active worlds
-	UFUNCTION(Category = Mix, meta = (CallInEditor = "true"))
+	// Activates this mix in all active worlds.
+	UFUNCTION(Category = Mix, meta = (CallInEditor = "true", DisplayPriority = "10"))
 	void ActivateMix();
 
-	// Deactivates this mix in all active worlds
-	UFUNCTION(Category = Mix, meta = (CallInEditor = "true"))
+	// Deactivates this mix in all active worlds. The mix is fully deactivated once all stages have finished their release times.
+	UFUNCTION(Category = Mix, meta = (CallInEditor = "true", DisplayPriority = "20"))
 	void DeactivateMix();
 
 	// Deactivates all mixes in all active worlds
-	UFUNCTION(Category = Mix, meta = (CallInEditor = "true"))
+	UFUNCTION(Category = Mix, meta = (CallInEditor = "true", DisplayPriority = "30"))
 	void DeactivateAllMixes();
 
 public:
-	UPROPERTY(EditAnywhere, Transient, Category = Mix)
-	uint32 ProfileIndex;
-
+	UPROPERTY(EditAnywhere, Transient, Category = Config, meta = (DisplayPriority = "30"))
+	uint32 ProfileIndex = 0;
+
+	// Once activated, the mix will start a timer for the given duration (seconds).
+	// When the timer ends, the mix will be deactivated.
+	// As a result, the attack time is included in this duration, but not the release time.
+	// When set to 0, the mix is activated and then immediately deactivated.
+	// If less than 0, the mix will remain activated until manually deactivated.
+	UPROPERTY(EditAnywhere, Category = Config, meta = (DisplayPriority = "10"), BlueprintReadOnly)
+	double Duration = -1.0;
+
+	// If a Mix is already active and you activate it again, one of two things will happen:
+	// If set to true, the stages will all go back to their default values and the mix will activate again,
+	// allowing the attack to trigger again.
+	// If set to false, calling activate will only reset the timer to deactivate (based on the Duration value).
+	UPROPERTY(EditAnywhere, Category = Config, meta = (DisplayPriority = "20"), BlueprintReadOnly)
+	bool bRetriggerOnActivation = false;
+	
 	virtual void BeginDestroy() override;
 
 #if WITH_EDITOR
@@ -78,7 +93,7 @@ public:
 #endif // WITH_EDITOR
 
 	/* Array of stages controlled by mix. */
-	UPROPERTY(EditAnywhere, Category = Mix, BlueprintReadOnly)
+	UPROPERTY(EditAnywhere, Category = Stages, BlueprintReadOnly)
 	TArray<FSoundControlBusMixStage> MixStages;
 };
 
