Fix Linker Error in Audio Modulation Statics

diff --git a/Engine/Plugins/Runtime/AudioModulation/Source/AudioModulation/Public/AudioModulationStatics.h b/Engine/Plugins/Runtime/AudioModulation/Source/AudioModulation/Public/AudioModulationStatics.h
--- a/Engine/Plugins/Runtime/AudioModulation/Source/AudioModulation/Public/AudioModulationStatics.h
+++ b/Engine/Plugins/Runtime/AudioModulation/Source/AudioModulation/Public/AudioModulationStatics.h
@@ -304,7 +304,7 @@ public:
 		AdvancedDisplay = "3",
 		Keywords = "modulation modulator stage")
 	)
-	static void UpdateMix(const UObject* WorldContextObject, USoundControlBusMix* Mix, TArray<FSoundControlBusMixStage> Stages, float FadeTime = -1.0f, double Duration = -1.0, bool bRetriggerOnActivation = false);
+	static UE_API void UpdateMix(const UObject* WorldContextObject, USoundControlBusMix* Mix, TArray<FSoundControlBusMixStage> Stages, float FadeTime = -1.0f, double Duration = -1.0, bool bRetriggerOnActivation = false);
 
 	/** Sets a Global Control Bus Mix with a single stage associated with the provided Bus to the given float value.  This call should
 	 * be reserved for buses that are to be always active. It is *NOT* recommended for transient buses, as not calling clear can keep
