Revert Activate Bus Mix behavior to overwrite runtime settings

diff --git a/Engine/Plugins/Runtime/AudioModulation/Source/AudioModulation/Private/AudioModulationSystem.cpp b/Engine/Plugins/Runtime/AudioModulation/Source/AudioModulation/Private/AudioModulationSystem.cpp
--- a/Engine/Plugins/Runtime/AudioModulation/Source/AudioModulation/Private/AudioModulationSystem.cpp
+++ b/Engine/Plugins/Runtime/AudioModulation/Source/AudioModulation/Private/AudioModulationSystem.cpp
@@ -166,7 +166,7 @@ namespace AudioModulation
 			FBusMixHandle BusMixHandle = FBusMixHandle::Get(Settings.GetId(), RefProxies.BusMixes);
 			if (BusMixHandle.IsValid())
 			{
-				BusMixHandle.FindProxy().SetEnabled();
+				BusMixHandle.FindProxy().SetMixDataAndEnable(MoveTemp(Settings));
 			}
 			else
 			{
