Move Mix Deactivate Trace Log to Mix Proxy [on behalf of alfaroh.corneyiii]

diff --git a/Engine/Plugins/Runtime/AudioModulation/Source/AudioModulation/Private/AudioModulationSystem.cpp b/Engine/Plugins/Runtime/AudioModulation/Source/AudioModulation/Private/AudioModulationSystem.cpp
--- a/Engine/Plugins/Runtime/AudioModulation/Source/AudioModulation/Private/AudioModulationSystem.cpp
+++ b/Engine/Plugins/Runtime/AudioModulation/Source/AudioModulation/Private/AudioModulationSystem.cpp
@@ -85,11 +85,7 @@ UE_TRACE_EVENT_BEGIN(Audio, BusFinalValuesUpdate)
 	UE_TRACE_EVENT_FIELD(float[], BusValues)
 UE_TRACE_EVENT_END()
 
-UE_TRACE_EVENT_BEGIN(Audio, ModulatingSourceDeactivate)
-	UE_TRACE_EVENT_FIELD(uint32, DeviceId)
-	UE_TRACE_EVENT_FIELD(uint32, SourceId)
-	UE_TRACE_EVENT_FIELD(double, Timestamp)
-UE_TRACE_EVENT_END()
+UE_TRACE_EVENT_DEFINE(Audio, ModulatingSourceDeactivate)
 #endif // UE_AUDIO_PROFILERTRACE_ENABLED
 
 namespace AudioModulation
@@ -268,13 +264,6 @@ namespace AudioModulation
 			{
 				FModulatorBusMixProxy& MixProxy = MixHandle.FindProxy();
 				MixProxy.SetStopping();
-
-#if UE_AUDIO_PROFILERTRACE_ENABLED
-				UE_TRACE_LOG(Audio, ModulatingSourceDeactivate, AudioChannel)
-					<< ModulatingSourceDeactivate.DeviceId(AudioDeviceId)
-					<< ModulatingSourceDeactivate.SourceId(BusMixId)
-					<< ModulatingSourceDeactivate.Timestamp(FPlatformTime::Cycles64());
-#endif
 			}
 		});
 	}
diff --git a/Engine/Plugins/Runtime/AudioModulation/Source/AudioModulation/Private/AudioModulationSystem.h b/Engine/Plugins/Runtime/AudioModulation/Source/AudioModulation/Private/AudioModulationSystem.h
--- a/Engine/Plugins/Runtime/AudioModulation/Source/AudioModulation/Private/AudioModulationSystem.h
+++ b/Engine/Plugins/Runtime/AudioModulation/Source/AudioModulation/Private/AudioModulationSystem.h
@@ -60,6 +60,13 @@ namespace AudioModulation
 	};
 } // namespace AudioModulation
 
+#if UE_AUDIO_PROFILERTRACE_ENABLED
+UE_TRACE_EVENT_BEGIN_EXTERN(Audio, ModulatingSourceDeactivate)
+	UE_TRACE_EVENT_FIELD(uint32, DeviceId)
+	UE_TRACE_EVENT_FIELD(uint32, SourceId)
+	UE_TRACE_EVENT_FIELD(double, Timestamp)
+UE_TRACE_EVENT_END()
+#endif
 
 #if WITH_AUDIOMODULATION
 namespace AudioModulation
diff --git a/Engine/Plugins/Runtime/AudioModulation/Source/AudioModulation/Private/SoundControlBusMixProxy.cpp b/Engine/Plugins/Runtime/AudioModulation/Source/AudioModulation/Private/SoundControlBusMixProxy.cpp
--- a/Engine/Plugins/Runtime/AudioModulation/Source/AudioModulation/Private/SoundControlBusMixProxy.cpp
+++ b/Engine/Plugins/Runtime/AudioModulation/Source/AudioModulation/Private/SoundControlBusMixProxy.cpp
@@ -4,6 +4,7 @@
 #include "Audio/AudioAddressPattern.h"
 #include "AudioDevice.h"
 #include "AudioDeviceManager.h"
+#include "AudioMixerTrace.h"
 #include "AudioModulation.h"
 #include "AudioModulationLogging.h"
 #include "AudioModulationProfileSerializer.h"
@@ -269,6 +270,13 @@ namespace AudioModulation
 		if (bRequestStop)
 		{
 			Status = EStatus::Stopped;
+			
+#if	UE_AUDIO_PROFILERTRACE_ENABLED
+			UE_TRACE_LOG(Audio, ModulatingSourceDeactivate, AudioChannel)
+				<< ModulatingSourceDeactivate.DeviceId(ModSystem->GetAudioDeviceId())
+				<< ModulatingSourceDeactivate.SourceId(GetId())
+				<< ModulatingSourceDeactivate.Timestamp(FPlatformTime::Cycles64());
+#endif
 		}
 	}
 } // namespace AudioModulation
