[Audio Insights] Add r-click menu options to browse to and edit asset in Control Bus and Modulation tabs

diff --git a/Engine/Plugins/AudioInsights/Source/AudioInsights/Public/Views/TableDashboardViewFactory.h b/Engine/Plugins/AudioInsights/Source/AudioInsights/Public/Views/TableDashboardViewFactory.h
--- a/Engine/Plugins/AudioInsights/Source/AudioInsights/Public/Views/TableDashboardViewFactory.h
+++ b/Engine/Plugins/AudioInsights/Source/AudioInsights/Public/Views/TableDashboardViewFactory.h
@@ -185,10 +185,7 @@ namespace UE::Audio::Insights
 	protected:
 #if WITH_EDITOR
 		UE_API virtual TSharedRef<SWidget> MakeAssetMenuBar() const;
-#endif // WITH_EDITOR
 
-	private:
-#if WITH_EDITOR
 		UE_API TArray<UObject*> GetSelectedEditableAssets() const;
 
 		UE_API bool OpenAsset() const;
diff --git a/Engine/Plugins/Runtime/AudioModulation/Source/AudioModulation/Private/AudioModulationSystem.cpp b/Engine/Plugins/Runtime/AudioModulation/Source/AudioModulation/Private/AudioModulationSystem.cpp
--- a/Engine/Plugins/Runtime/AudioModulation/Source/AudioModulation/Private/AudioModulationSystem.cpp
+++ b/Engine/Plugins/Runtime/AudioModulation/Source/AudioModulation/Private/AudioModulationSystem.cpp
@@ -179,14 +179,14 @@ namespace AudioModulation
 					<< BusMixRegisterBus.SourceId(StageSetting.BusSettings.GetId())
 					<< BusMixRegisterBus.Timestamp(FPlatformTime::Cycles64())
 					<< BusMixRegisterBus.ModulatingSourceId(Settings.GetId())
-					<< BusMixRegisterBus.BusName(*(StageSetting.BusSettings.GetName().ToString()));
+					<< BusMixRegisterBus.BusName(*(StageSetting.BusSettings.GetPath().ToString()));
 			}
 
 			UE_TRACE_LOG(Audio, BusMixActivate, AudioChannel)
 				<< BusMixActivate.DeviceId(AudioDeviceId)
 				<< BusMixActivate.SourceId(Settings.GetId())
 				<< BusMixActivate.Timestamp(FPlatformTime::Cycles64())
-				<< BusMixActivate.Name(*Settings.GetName().ToString());
+				<< BusMixActivate.Name(*Settings.GetPath().ToString());
 #endif
 		});
 	}
@@ -655,7 +655,7 @@ namespace AudioModulation
 					<< ControlBusUpdate.DeviceId(static_cast<uint32>(AudioDeviceId))
 					<< ControlBusUpdate.ControlBusId(static_cast<uint32>(BusId))
 					<< ControlBusUpdate.Timestamp(FPlatformTime::Cycles64())
-					<< ControlBusUpdate.Name(*(ControlBusProxy.GetName().ToString()))
+					<< ControlBusUpdate.Name(*(ControlBusProxy.GetPath().ToString()))
 					<< ControlBusUpdate.ParamName(*(ControlBusProxy.GetParameterName().ToString()))
 					<< ControlBusUpdate.Value(ControlBusProxy.GetValue());
 
@@ -836,7 +836,7 @@ namespace AudioModulation
 					<< BusMixRegisterBus.SourceId(ControlBusProxy.GetId())
 					<< BusMixRegisterBus.Timestamp(FPlatformTime::Cycles64())
 					<< BusMixRegisterBus.ModulatingSourceId(BusId)
-					<< BusMixRegisterBus.BusName(*(ControlBusProxy.GetName().ToString()));
+					<< BusMixRegisterBus.BusName(*(ControlBusProxy.GetPath().ToString()));
 
 				ControlBusProxy.OnTraceStarted(*this);
 			}
@@ -847,7 +847,7 @@ namespace AudioModulation
 					<< BusMixActivate.DeviceId(AudioDeviceId)
 					<< BusMixActivate.SourceId(BusMixId)
 					<< BusMixActivate.Timestamp(FPlatformTime::Cycles64())
-					<< BusMixActivate.Name(*(ModulatorBusMixProxy.GetName().ToString()));
+					<< BusMixActivate.Name(*(ModulatorBusMixProxy.GetPath().ToString()));
 			}
 		});
 #endif // UE_AUDIO_PROFILERTRACE_ENABLED
diff --git a/Engine/Plugins/Runtime/AudioModulation/Source/AudioModulation/Private/SoundControlBusMixProxy.cpp b/Engine/Plugins/Runtime/AudioModulation/Source/AudioModulation/Private/SoundControlBusMixProxy.cpp
--- a/Engine/Plugins/Runtime/AudioModulation/Source/AudioModulation/Private/SoundControlBusMixProxy.cpp
+++ b/Engine/Plugins/Runtime/AudioModulation/Source/AudioModulation/Private/SoundControlBusMixProxy.cpp
@@ -21,7 +21,7 @@ namespace AudioModulation
 	const FBusMixId InvalidBusMixId = INDEX_NONE;
 
 	FModulatorBusMixStageSettings::FModulatorBusMixStageSettings(const FSoundControlBusMixStage& InStage)
-		: TModulatorBase<FBusId>(InStage.Bus->GetFName(), InStage.Bus->GetUniqueID())
+		: TModulatorBase<FBusId>(InStage.Bus->GetPathName(), InStage.Bus->GetUniqueID())
 		, Address(InStage.Bus->Address)
 		, ParamClassId(INDEX_NONE)
 		, ParamId(INDEX_NONE)
@@ -41,7 +41,7 @@ namespace AudioModulation
 	}
 
 	FModulatorBusMixStageProxy::FModulatorBusMixStageProxy(FModulatorBusMixStageSettings&& InSettings, FAudioModulationSystem& OutModSystem)
-		: TModulatorBase<FBusId>(InSettings.BusSettings.GetName(), InSettings.BusSettings.GetId())
+		: TModulatorBase<FBusId>(InSettings.BusSettings.GetPath(), InSettings.BusSettings.GetId())
 		, Address(MoveTemp(InSettings.Address))
 		, ParamClassId(InSettings.ParamClassId)
 		, ParamId(InSettings.ParamId)
@@ -51,7 +51,7 @@ namespace AudioModulation
 	}
 
 	FModulatorBusMixSettings::FModulatorBusMixSettings(const USoundControlBusMix& InBusMix)
-		: TModulatorBase<FBusMixId>(InBusMix.GetFName()
+		: TModulatorBase<FBusMixId>(InBusMix.GetPathName()
 		, InBusMix.GetUniqueID())
 		, Duration(InBusMix.Duration)
 		, bRetriggerOnActivation(InBusMix.bRetriggerOnActivation)
@@ -73,7 +73,7 @@ namespace AudioModulation
 	}
 
 	FModulatorBusMixSettings::FModulatorBusMixSettings(FModulatorBusMixSettings&& InBusMixSettings)
-		: TModulatorBase<FBusMixId>(InBusMixSettings.GetName(), InBusMixSettings.GetId())
+		: TModulatorBase<FBusMixId>(InBusMixSettings.GetPath(), InBusMixSettings.GetId())
 		, Stages(MoveTemp(InBusMixSettings.Stages))
 		, Duration(InBusMixSettings.Duration)
 		, bRetriggerOnActivation(InBusMixSettings.bRetriggerOnActivation)
@@ -81,7 +81,7 @@ namespace AudioModulation
 	}
 
 	FModulatorBusMixProxy::FModulatorBusMixProxy(FModulatorBusMixSettings&& InSettings, FAudioModulationSystem& OutModSystem)
-		: TModulatorProxyRefType(InSettings.GetName(), InSettings.GetId(), OutModSystem)
+		: TModulatorProxyRefType(InSettings.GetPath(), InSettings.GetId(), OutModSystem)
 		, Duration(InSettings.Duration)
 		, TimeRemaining(Duration)
 		, bRetriggerOnActivation(InSettings.bRetriggerOnActivation)
diff --git a/Engine/Plugins/Runtime/AudioModulation/Source/AudioModulation/Private/SoundControlBusProxy.cpp b/Engine/Plugins/Runtime/AudioModulation/Source/AudioModulation/Private/SoundControlBusProxy.cpp
--- a/Engine/Plugins/Runtime/AudioModulation/Source/AudioModulation/Private/SoundControlBusProxy.cpp
+++ b/Engine/Plugins/Runtime/AudioModulation/Source/AudioModulation/Private/SoundControlBusProxy.cpp
@@ -58,13 +58,13 @@ namespace AudioModulation
 				<< GeneratorRegisterBus.SourceId(GetId())
 				<< GeneratorRegisterBus.Timestamp(FPlatformTime::Cycles64())
 				<< GeneratorRegisterBus.ModulatingSourceId(GeneratorSetting.GetId())
-				<< GeneratorRegisterBus.BusName(*(GetName().ToString()));
+				<< GeneratorRegisterBus.BusName(*(GetPath().ToString()));
 
 			UE_TRACE_LOG(Audio, GeneratorActivate, AudioChannel)
 				<< GeneratorActivate.DeviceId(ModSystem.GetAudioDeviceId())
 				<< GeneratorActivate.SourceId(GeneratorSetting.GetId())
 				<< GeneratorActivate.Timestamp(FPlatformTime::Cycles64())
-				<< GeneratorActivate.Name(*GeneratorSetting.GetName().ToString());
+				<< GeneratorActivate.Name(*GeneratorSetting.GetPath().ToString());
 		}
 #endif
 
@@ -80,7 +80,7 @@ namespace AudioModulation
 	}
 
 	FControlBusProxy::FControlBusProxy(FControlBusSettings&& InSettings, FAudioModulationSystem& InModSystem)
-		: TModulatorProxyRefType(InSettings.GetName(), InSettings.GetId(), InModSystem)
+		: TModulatorProxyRefType(InSettings.GetPath(), InSettings.GetId(), InModSystem)
 	{
 		Init(MoveTemp(InSettings));
 	}
@@ -169,7 +169,7 @@ namespace AudioModulation
 			<< ControlBusActivate.DeviceId(ModSystem->GetAudioDeviceId())
 			<< ControlBusActivate.ControlBusId(InSettings.GetId())
 			<< ControlBusActivate.Timestamp(FPlatformTime::Cycles64())
-			<< ControlBusActivate.Name(*(InSettings.GetName().ToString()))
+			<< ControlBusActivate.Name(*(InSettings.GetPath().ToString()))
 			<< ControlBusActivate.ParamName(*(GetParameterName().ToString()));
 #endif // UE_AUDIO_PROFILERTRACE_ENABLED
 	}
@@ -210,13 +210,13 @@ namespace AudioModulation
 				<< GeneratorRegisterBus.SourceId(GetId())
 				<< GeneratorRegisterBus.Timestamp(FPlatformTime::Cycles64())
 				<< GeneratorRegisterBus.ModulatingSourceId(GeneratorHandle.GetId())
-				<< GeneratorRegisterBus.BusName(*(GetName().ToString()));
+				<< GeneratorRegisterBus.BusName(*(GetPath().ToString()));
 
 			UE_TRACE_LOG(Audio, GeneratorActivate, AudioChannel)
 				<< GeneratorActivate.DeviceId(InModSystem.GetAudioDeviceId())
 				<< GeneratorActivate.SourceId(GeneratorHandle.GetId())
 				<< GeneratorActivate.Timestamp(FPlatformTime::Cycles64())
-				<< GeneratorActivate.Name(*(GeneratorHandle.FindProxy().GetName().ToString()));
+				<< GeneratorActivate.Name(*(GeneratorHandle.FindProxy().GetPath().ToString()));
 		}
 	}
 #endif // UE_AUDIO_PROFILERTRACE_ENABLED
diff --git a/Engine/Plugins/Runtime/AudioModulation/Source/AudioModulation/Private/SoundControlBusProxy.h b/Engine/Plugins/Runtime/AudioModulation/Source/AudioModulation/Private/SoundControlBusProxy.h
--- a/Engine/Plugins/Runtime/AudioModulation/Source/AudioModulation/Private/SoundControlBusProxy.h
+++ b/Engine/Plugins/Runtime/AudioModulation/Source/AudioModulation/Private/SoundControlBusProxy.h
@@ -28,7 +28,7 @@ namespace AudioModulation
 		Audio::FModulationParameter OutputParameter;
 
 		FControlBusSettings(const USoundControlBus& InBus)
-			: TModulatorBase<FBusId>(InBus.GetFName(), InBus.GetUniqueID())
+			: TModulatorBase<FBusId>(InBus.GetPathName(), InBus.GetUniqueID())
 			, bBypass(InBus.bBypass)
 			, DefaultValue(InBus.GetDefaultNormalizedValue())
 			, MixFunction(InBus.GetMixFunction())
diff --git a/Engine/Plugins/Runtime/AudioModulation/Source/AudioModulation/Private/SoundModulationGeneratorProxy.cpp b/Engine/Plugins/Runtime/AudioModulation/Source/AudioModulation/Private/SoundModulationGeneratorProxy.cpp
--- a/Engine/Plugins/Runtime/AudioModulation/Source/AudioModulation/Private/SoundModulationGeneratorProxy.cpp
+++ b/Engine/Plugins/Runtime/AudioModulation/Source/AudioModulation/Private/SoundModulationGeneratorProxy.cpp
@@ -22,7 +22,7 @@ namespace AudioModulation
 	}
 
 	FModulatorGeneratorProxy::FModulatorGeneratorProxy(FModulationGeneratorSettings&& InSettings, FAudioModulationSystem& InModSystem)
-		: TModulatorProxyRefType(InSettings.GetName(), InSettings.GetId(), InModSystem)
+		: TModulatorProxyRefType(InSettings.GetPath(), InSettings.GetId(), InModSystem)
 		, Generator(MoveTemp(InSettings.Generator))
 	{
 		Generator->Init(InModSystem.GetAudioDeviceId());
diff --git a/Engine/Plugins/Runtime/AudioModulation/Source/AudioModulation/Private/SoundModulationGeneratorProxy.h b/Engine/Plugins/Runtime/AudioModulation/Source/AudioModulation/Private/SoundModulationGeneratorProxy.h
--- a/Engine/Plugins/Runtime/AudioModulation/Source/AudioModulation/Private/SoundModulationGeneratorProxy.h
+++ b/Engine/Plugins/Runtime/AudioModulation/Source/AudioModulation/Private/SoundModulationGeneratorProxy.h
@@ -30,19 +30,19 @@ namespace AudioModulation
 		FModulationGeneratorSettings() = default;
 
 		FModulationGeneratorSettings(const USoundModulationGenerator& InGenerator)
-			: TModulatorBase<FGeneratorId>(InGenerator.GetFName(), InGenerator.GetUniqueID())
+			: TModulatorBase<FGeneratorId>(InGenerator.GetPathName(), InGenerator.GetUniqueID())
 			, Generator(InGenerator.CreateInstance())
 		{
 		}
 
 		FModulationGeneratorSettings(const FModulationGeneratorSettings& InSettings)
-			: TModulatorBase<FGeneratorId>(InSettings.GetName(), InSettings.GetId())
+			: TModulatorBase<FGeneratorId>(InSettings.GetPath(), InSettings.GetId())
 			, Generator(InSettings.Generator.IsValid() ? InSettings.Generator->Clone() : nullptr)
 		{
 		}
 
 		FModulationGeneratorSettings(FModulationGeneratorSettings&& InSettings)
-			: TModulatorBase<FGeneratorId>(InSettings.GetName(), InSettings.GetId())
+			: TModulatorBase<FGeneratorId>(InSettings.GetPath(), InSettings.GetId())
 			, Generator(MoveTemp(InSettings.Generator))
 		{
 		}
diff --git a/Engine/Plugins/Runtime/AudioModulation/Source/AudioModulation/Private/SoundModulationPatchProxy.cpp b/Engine/Plugins/Runtime/AudioModulation/Source/AudioModulation/Private/SoundModulationPatchProxy.cpp
--- a/Engine/Plugins/Runtime/AudioModulation/Source/AudioModulation/Private/SoundModulationPatchProxy.cpp
+++ b/Engine/Plugins/Runtime/AudioModulation/Source/AudioModulation/Private/SoundModulationPatchProxy.cpp
@@ -127,7 +127,7 @@ namespace AudioModulation
 	}
 
 	FModulationPatchRefProxy::FModulationPatchRefProxy(FModulationPatchSettings&& InSettings, FAudioModulationSystem& OutModSystem)
-		: TModulatorProxyRefType(InSettings.GetName(), InSettings.GetId(), OutModSystem)
+		: TModulatorProxyRefType(InSettings.GetPath(), InSettings.GetId(), OutModSystem)
 		, FModulationPatchProxy(MoveTemp(InSettings), OutModSystem)
 	{
 	}
diff --git a/Engine/Plugins/Runtime/AudioModulation/Source/AudioModulation/Private/SoundModulationPatchProxy.h b/Engine/Plugins/Runtime/AudioModulation/Source/AudioModulation/Private/SoundModulationPatchProxy.h
--- a/Engine/Plugins/Runtime/AudioModulation/Source/AudioModulation/Private/SoundModulationPatchProxy.h
+++ b/Engine/Plugins/Runtime/AudioModulation/Source/AudioModulation/Private/SoundModulationPatchProxy.h
@@ -83,7 +83,7 @@ namespace AudioModulation
 		FModulationPatchSettings(const FModulationPatchSettings& InPatchSettings) = default;
 
 		FModulationPatchSettings(const USoundModulationPatch& InPatch)
-			: TModulatorBase<FPatchId>(InPatch.GetFName(), InPatch.GetUniqueID())
+			: TModulatorBase<FPatchId>(InPatch.GetPathName(), InPatch.GetUniqueID())
 			, bBypass(InPatch.PatchSettings.bBypass)
 			, OutputParameter(InPatch.GetOutputParameter())
 		{
diff --git a/Engine/Plugins/Runtime/AudioModulation/Source/AudioModulation/Private/SoundModulationProxy.h b/Engine/Plugins/Runtime/AudioModulation/Source/AudioModulation/Private/SoundModulationProxy.h
--- a/Engine/Plugins/Runtime/AudioModulation/Source/AudioModulation/Private/SoundModulationProxy.h
+++ b/Engine/Plugins/Runtime/AudioModulation/Source/AudioModulation/Private/SoundModulationProxy.h
@@ -238,6 +238,7 @@ namespace AudioModulation
 
 #if !UE_BUILD_SHIPPING
 		FName Name;
+		FName Path;
 #endif // !UE_BUILD_SHIPPING
 
 	public:
@@ -246,21 +247,23 @@ namespace AudioModulation
 		{
 		}
 
-		TModulatorBase(const FName InName, const uint32 InId)
+		TModulatorBase(const FName InPath, const uint32 InId)
 			: Id(static_cast<IdType>(InId))
 #if !UE_BUILD_SHIPPING
-			, Name(InName)
+			, Name(FSoftObjectPath(InPath.ToString()).GetAssetName())
+			, Path(InPath)
+#endif // !UE_BUILD_SHIPPING
+		{
+		}
+
+		TModulatorBase(const FString& InPath, const uint32 InId)
+			: Id(static_cast<IdType>(InId))
+#if !UE_BUILD_SHIPPING
+			, Name(FSoftObjectPath(InPath).GetAssetName())
+			, Path(InPath)
 #endif // !UE_BUILD_SHIPPING
 		{
 		}
-		
-		TModulatorBase(const FString& InName, const uint32 InId)
-        			: Id(static_cast<IdType>(InId))
-        #if !UE_BUILD_SHIPPING
-        			, Name(InName)
-        #endif // !UE_BUILD_SHIPPING
-        {
-        }
 
 		virtual ~TModulatorBase() = default;
 
@@ -279,6 +282,17 @@ namespace AudioModulation
 
 			return Name;
 		}
+
+		// FOR DEBUG USE ONLY (Not available in shipped builds):
+		// Provides path of object that generated proxy.
+		const FName& GetPath() const
+		{
+#if UE_BUILD_SHIPPING
+			static FName Path;
+#endif // UE_BUILD_SHIPPING
+
+			return Path;
+		}
 	};
 
 	template<typename IdType, typename ProxyType, typename ProxySettings>
@@ -296,22 +310,22 @@ namespace AudioModulation
 		{
 		}
 
-		TModulatorProxyRefType(const FName& InName, const uint32 InId, FAudioModulationSystem& OutModSystem)
-			: TModulatorBase<IdType>(InName, InId)
+		TModulatorProxyRefType(const FName& InPath, const uint32 InId, FAudioModulationSystem& OutModSystem)
+			: TModulatorBase<IdType>(InPath, InId)
 			, RefCount(0)
 			, ModSystem(&OutModSystem)
 		{
 		}
 
 		TModulatorProxyRefType(const TModulatorProxyRefType& InProxyRef)
-			: TModulatorBase<IdType>(InProxyRef.GetName(), InProxyRef.GetId())
+			: TModulatorBase<IdType>(InProxyRef.GetPath(), InProxyRef.GetId())
 			, RefCount(InProxyRef.RefCount)
 			, ModSystem(InProxyRef.ModSystem)
 		{
 		}
 
 		TModulatorProxyRefType(TModulatorProxyRefType&& InProxyRef)
-			: TModulatorBase<IdType>(InProxyRef.GetName(), InProxyRef.GetId())
+			: TModulatorBase<IdType>(InProxyRef.GetPath(), InProxyRef.GetId())
 			, RefCount(InProxyRef.RefCount)
 			, ModSystem(InProxyRef.ModSystem)
 		{
diff --git a/Engine/Plugins/Runtime/AudioModulation/Source/AudioModulationEditor/Private/Insights/Messages/ControlBusTraceMessages.h b/Engine/Plugins/Runtime/AudioModulation/Source/AudioModulationEditor/Private/Insights/Messages/ControlBusTraceMessages.h
--- a/Engine/Plugins/Runtime/AudioModulation/Source/AudioModulationEditor/Private/Insights/Messages/ControlBusTraceMessages.h
+++ b/Engine/Plugins/Runtime/AudioModulation/Source/AudioModulationEditor/Private/Insights/Messages/ControlBusTraceMessages.h
@@ -68,10 +68,12 @@ namespace AudioModulationEditor
 		FControlBusDashboardEntry() = default;
 		virtual ~FControlBusDashboardEntry() = default;
 
+		FText GetDisplayNameAsFText() const { return FText::FromString(DisplayName); }
 		FText GetParamNameAsFText() const { return FText::FromString(ParamName); }
 
 		FBusId ControlBusId = INDEX_NONE;
 		float Value = 1.0f;
+		FString DisplayName;
 		FString ParamName;
 	};
 
diff --git a/Engine/Plugins/Runtime/AudioModulation/Source/AudioModulationEditor/Private/Insights/Messages/ModulationMatrixTraceMessages.h b/Engine/Plugins/Runtime/AudioModulation/Source/AudioModulationEditor/Private/Insights/Messages/ModulationMatrixTraceMessages.h
--- a/Engine/Plugins/Runtime/AudioModulation/Source/AudioModulationEditor/Private/Insights/Messages/ModulationMatrixTraceMessages.h
+++ b/Engine/Plugins/Runtime/AudioModulation/Source/AudioModulationEditor/Private/Insights/Messages/ModulationMatrixTraceMessages.h
@@ -30,6 +30,7 @@ namespace AudioModulationEditor
 
 		FSourceId SourceId = INDEX_NONE;
 		EModulationMatrixEntryType EntryType = EModulationMatrixEntryType::None;
+		FString DisplayName;
 
 		TMap<FBusId, float> BusIdToValueMap;
 	};
diff --git a/Engine/Plugins/Runtime/AudioModulation/Source/AudioModulationEditor/Private/Insights/Providers/ControlBusTraceProvider.cpp b/Engine/Plugins/Runtime/AudioModulation/Source/AudioModulationEditor/Private/Insights/Providers/ControlBusTraceProvider.cpp
--- a/Engine/Plugins/Runtime/AudioModulation/Source/AudioModulationEditor/Private/Insights/Providers/ControlBusTraceProvider.cpp
+++ b/Engine/Plugins/Runtime/AudioModulation/Source/AudioModulationEditor/Private/Insights/Providers/ControlBusTraceProvider.cpp
@@ -38,6 +38,7 @@ namespace AudioModulationEditor
 			FControlBusDashboardEntry& EntryRef = *OutEntry->Get();
 			EntryRef.Name = *Msg.BusName;
 			EntryRef.ControlBusId = Msg.ControlBusId;
+			EntryRef.DisplayName = EntryRef.GetDisplayName().ToString();
 			EntryRef.ParamName = Msg.ParamName;
 		});
 
@@ -46,6 +47,7 @@ namespace AudioModulationEditor
 		{
 			FControlBusDashboardEntry& EntryRef = *OutEntry->Get();
 			EntryRef.Name = *Msg.BusName;
+			EntryRef.DisplayName = EntryRef.GetDisplayName().ToString();
 			EntryRef.ParamName = Msg.ParamName;
 			EntryRef.Value = Msg.Value;
 		});
diff --git a/Engine/Plugins/Runtime/AudioModulation/Source/AudioModulationEditor/Private/Insights/Providers/ModulationMatrixTraceProvider.cpp b/Engine/Plugins/Runtime/AudioModulation/Source/AudioModulationEditor/Private/Insights/Providers/ModulationMatrixTraceProvider.cpp
--- a/Engine/Plugins/Runtime/AudioModulation/Source/AudioModulationEditor/Private/Insights/Providers/ModulationMatrixTraceProvider.cpp
+++ b/Engine/Plugins/Runtime/AudioModulation/Source/AudioModulationEditor/Private/Insights/Providers/ModulationMatrixTraceProvider.cpp
@@ -127,11 +127,13 @@ namespace AudioModulationEditor
 			{
 				if (TMap<FBusId, FString>* FoundBusInfoPtr = ModulatingSourceIdToBusesInfoMap.Find(Msg.ModulatingSourceId))
 				{
-					FoundBusInfoPtr->FindOrAdd(Msg.SourceId, Msg.BusName);
+					FString AssetName = FSoftObjectPath(Msg.BusName).GetAssetName();
+					FoundBusInfoPtr->FindOrAdd(Msg.SourceId, AssetName.IsEmpty() ? Msg.BusName : MoveTemp(AssetName));
 				}
 				else
 				{
-					ModulatingSourceIdToBusesInfoMap.Emplace(Msg.ModulatingSourceId).Emplace(Msg.SourceId, Msg.BusName);
+					FString AssetName = FSoftObjectPath(Msg.BusName).GetAssetName();
+					ModulatingSourceIdToBusesInfoMap.Emplace(Msg.ModulatingSourceId).Emplace(Msg.SourceId, AssetName.IsEmpty() ? Msg.BusName : MoveTemp(AssetName));
 				}
 			});
 
@@ -140,6 +142,7 @@ namespace AudioModulationEditor
 			{
 				FModulationMatrixDashboardEntry& EntryRef = *OutEntry->Get();
 				EntryRef.Name = *Msg.Name;
+				EntryRef.DisplayName = EntryRef.GetDisplayName().ToString();
 				EntryRef.EntryType = EModulationMatrixEntryType::BusMix;
 
 				UpdateActiveModulatorSourceIds(Msg.DeviceId, EntryRef.SourceId);
@@ -150,6 +153,7 @@ namespace AudioModulationEditor
 			{
 				FModulationMatrixDashboardEntry& EntryRef = *OutEntry->Get();
 				EntryRef.Name = *Msg.Name;
+				EntryRef.DisplayName = EntryRef.GetDisplayName().ToString();
 				EntryRef.EntryType = EModulationMatrixEntryType::Generator;
 
 				UpdateActiveModulatorSourceIds(Msg.DeviceId, EntryRef.SourceId);
@@ -203,7 +207,8 @@ namespace AudioModulationEditor
 			if (EntryRef.Name.IsEmpty())
 			{
 				static const FString BusFinalValuesEntryName("Final Values: ");
-				EntryRef.Name      = BusFinalValuesEntryName;
+				EntryRef.Name = BusFinalValuesEntryName;
+				EntryRef.DisplayName = BusFinalValuesEntryName;
 				EntryRef.EntryType = EModulationMatrixEntryType::BusFinalValues;
 			}
 
diff --git a/Engine/Plugins/Runtime/AudioModulation/Source/AudioModulationEditor/Private/Insights/Views/AudioModulationEditorCommands.cpp b/Engine/Plugins/Runtime/AudioModulation/Source/AudioModulationEditor/Private/Insights/Views/AudioModulationEditorCommands.cpp
new file mode 100644
--- /dev/null
+++ b/Engine/Plugins/Runtime/AudioModulation/Source/AudioModulationEditor/Private/Insights/Views/AudioModulationEditorCommands.cpp
@@ -0,0 +1,24 @@
+// Copyright Epic Games, Inc. All Rights Reserved.
+#include "AudioModulationEditorCommands.h"
+
+#include "AudioInsightsStyle.h"
+#include "Framework/Commands/InputChord.h"
+#include "Framework/Commands/UICommandInfo.h"
+
+#define LOCTEXT_NAMESPACE "AudioModulationInsights"
+
+namespace AudioModulationEditor
+{
+	FAudioModulationEditorCommands::FAudioModulationEditorCommands()
+		: TCommands<FAudioModulationEditorCommands>("AudioModulationEditorCommands", LOCTEXT("AudioModulationEditorCommands_ContextDescText", "Audio Modulation Editor Commands"), NAME_None, UE::Audio::Insights::FSlateStyle::GetStyleName())
+	{
+	}
+	
+	void FAudioModulationEditorCommands::RegisterCommands()
+	{
+		UI_COMMAND(Browse, "Browse To Asset", "Browses to the selected sound asset in the content browser.", EUserInterfaceActionType::Button, FInputChord(EKeys::B, EModifierKey::Control));
+		UI_COMMAND(Edit, "Edit", "Opens the selected sound for edit.", EUserInterfaceActionType::Button, FInputChord(EKeys::E, EModifierKey::Control));
+	}
+} // namespace AudioModulationEditor
+
+#undef LOCTEXT_NAMESPACE
diff --git a/Engine/Plugins/Runtime/AudioModulation/Source/AudioModulationEditor/Private/Insights/Views/AudioModulationEditorCommands.h b/Engine/Plugins/Runtime/AudioModulation/Source/AudioModulationEditor/Private/Insights/Views/AudioModulationEditorCommands.h
new file mode 100644
--- /dev/null
+++ b/Engine/Plugins/Runtime/AudioModulation/Source/AudioModulationEditor/Private/Insights/Views/AudioModulationEditorCommands.h
@@ -0,0 +1,25 @@
+// Copyright Epic Games, Inc. All Rights Reserved.
+#pragma once
+
+#include "Framework/Commands/Commands.h"
+#include "Templates/SharedPointer.h"
+
+class FUICommandInfo;
+
+namespace AudioModulationEditor
+{
+	class FAudioModulationEditorCommands : public TCommands<FAudioModulationEditorCommands>
+	{
+	public:
+		FAudioModulationEditorCommands();
+
+		virtual void RegisterCommands() override;
+
+		TSharedPtr<const FUICommandInfo> GetBrowseCommand() const { return Browse; }
+		TSharedPtr<const FUICommandInfo> GetEditCommand() const   { return Edit; }
+		
+	private:
+		TSharedPtr<FUICommandInfo> Browse;
+		TSharedPtr<FUICommandInfo> Edit;
+	};
+} // namespace UE::Audio::Insights
diff --git a/Engine/Plugins/Runtime/AudioModulation/Source/AudioModulationEditor/Private/Insights/Views/ControlBusDashboardViewFactory.cpp b/Engine/Plugins/Runtime/AudioModulation/Source/AudioModulationEditor/Private/Insights/Views/ControlBusDashboardViewFactory.cpp
--- a/Engine/Plugins/Runtime/AudioModulation/Source/AudioModulationEditor/Private/Insights/Views/ControlBusDashboardViewFactory.cpp
+++ b/Engine/Plugins/Runtime/AudioModulation/Source/AudioModulationEditor/Private/Insights/Views/ControlBusDashboardViewFactory.cpp
@@ -1,6 +1,7 @@
 // Copyright Epic Games, Inc. All Rights Reserved.
 #include "ControlBusDashboardViewFactory.h"
 
+#include "AudioModulationEditorCommands.h"
 #include "AssetRegistry/AssetData.h"
 #include "AssetRegistry/AssetRegistryModule.h"
 #include "AudioDefines.h"
@@ -10,6 +11,7 @@
 #include "AudioInsightsTraceModule.h"
 #include "AudioInsightsTraceProviderBase.h"
 #include "Editor.h"
+#include "Framework/MultiBox/MultiBoxBuilder.h"
 #include "Internationalization/Text.h"
 #include "Insights/Providers/ControlBusTraceProvider.h"
 #include "Subsystems/AssetEditorSubsystem.h"
@@ -40,10 +42,15 @@ namespace AudioModulationEditor
 		{
 			StaticCastSharedPtr<UE::Audio::Insights::FTraceProviderBase>(ControlBusProvider)
 		};
+
+		FAudioModulationEditorCommands::Register();
+
+		BindCommands();
 	}
 
 	FControlBusDashboardViewFactory::~FControlBusDashboardViewFactory()
 	{
+		FAudioModulationEditorCommands::Unregister();
 	}
 
 	FName FControlBusDashboardViewFactory::GetName() const
@@ -64,7 +71,7 @@ namespace AudioModulationEditor
 		FTraceTableDashboardViewFactory::FilterEntries<FControlBusTraceProvider>([&FilterString](const IDashboardDataViewEntry& Entry)
 		{
 			const FControlBusDashboardEntry& ControlBusEntry = ControlBusPrivate::CastEntry(Entry);
-			if (ControlBusEntry.Name.Contains(FilterString))
+			if (ControlBusEntry.DisplayName.Contains(FilterString))
 			{
 				return false;
 			}
@@ -130,7 +137,7 @@ namespace AudioModulationEditor
 					"Name",
 					{
 						LOCTEXT("ControlBus_NameColumnDisplayName", "Name"),
-						[](const IDashboardDataViewEntry& InData) { return FText::FromString(ControlBusPrivate::CastEntry(InData).Name); },
+						[](const IDashboardDataViewEntry& InData) { return ControlBusPrivate::CastEntry(InData).GetDisplayNameAsFText(); },
 						false /* bDefaultHidden */,
 						0.75f /* FillWidth */
 					}
@@ -195,7 +202,7 @@ namespace AudioModulationEditor
 					const FControlBusDashboardEntry& AData = ControlBusPrivate::CastEntry(*A.Get());
 					const FControlBusDashboardEntry& BData = ControlBusPrivate::CastEntry(*B.Get());
 
-					return AData.Name.ToLower() < BData.Name.ToLower();
+					return AData.DisplayName.ToLower() < BData.DisplayName.ToLower();
 				});
 			}
 			else if (SortMode == EColumnSortMode::Descending)
@@ -205,7 +212,7 @@ namespace AudioModulationEditor
 					const FControlBusDashboardEntry& AData = ControlBusPrivate::CastEntry(*A.Get());
 					const FControlBusDashboardEntry& BData = ControlBusPrivate::CastEntry(*B.Get());
 
-					return BData.Name.ToLower() < AData.Name.ToLower();
+					return BData.DisplayName.ToLower() < AData.DisplayName.ToLower();
 				});
 			}
 		}
@@ -267,6 +274,38 @@ namespace AudioModulationEditor
 		return SNew(SBox);
 	}
 
+	TSharedPtr<SWidget> FControlBusDashboardViewFactory::OnConstructContextMenu()
+	{
+		const FAudioModulationEditorCommands& Commands = FAudioModulationEditorCommands::Get();
+
+		constexpr bool bShouldCloseWindowAfterMenuSelection = true;
+		FMenuBuilder MenuBuilder(bShouldCloseWindowAfterMenuSelection, CommandList);
+
+		MenuBuilder.BeginSection("ControlBusDashboardActions", LOCTEXT("ControlBusActions_HeaderText", "Control Bus Options"));
+		{
+			MenuBuilder.AddMenuEntry(Commands.GetBrowseCommand(), NAME_None, TAttribute<FText>(), TAttribute<FText>(), UE::Audio::Insights::FSlateStyle::Get().CreateIcon("AudioInsights.Icon.SoundDashboard.Browse"));
+			MenuBuilder.AddMenuEntry(Commands.GetEditCommand(), NAME_None, TAttribute<FText>(), TAttribute<FText>(), UE::Audio::Insights::FSlateStyle::Get().CreateIcon("AudioInsights.Icon.SoundDashboard.Edit"));
+		}
+
+		MenuBuilder.EndSection();
+
+		return MenuBuilder.MakeWidget();
+	}
+
+	FReply FControlBusDashboardViewFactory::OnDataRowKeyInput(const FGeometry& InGeometry, const FKeyEvent& InKeyEvent) const
+	{
+		return (CommandList && CommandList->ProcessCommandBindings(InKeyEvent)) ? FReply::Handled() : FReply::Unhandled();
+	}
+
+	void FControlBusDashboardViewFactory::BindCommands()
+	{
+		CommandList = MakeShared<FUICommandList>();
+
+		const FAudioModulationEditorCommands& Commands = FAudioModulationEditorCommands::Get();
+
+		CommandList->MapAction(Commands.GetBrowseCommand(), FExecuteAction::CreateLambda([this]() { BrowseToAsset(); }));
+		CommandList->MapAction(Commands.GetEditCommand(), FExecuteAction::CreateLambda([this]() { OpenAsset(); }));
+	}
 } // namespace AudioModulationEditor
 
 #undef LOCTEXT_NAMESPACE
diff --git a/Engine/Plugins/Runtime/AudioModulation/Source/AudioModulationEditor/Private/Insights/Views/ControlBusDashboardViewFactory.h b/Engine/Plugins/Runtime/AudioModulation/Source/AudioModulationEditor/Private/Insights/Views/ControlBusDashboardViewFactory.h
--- a/Engine/Plugins/Runtime/AudioModulation/Source/AudioModulationEditor/Private/Insights/Views/ControlBusDashboardViewFactory.h
+++ b/Engine/Plugins/Runtime/AudioModulation/Source/AudioModulationEditor/Private/Insights/Views/ControlBusDashboardViewFactory.h
@@ -29,5 +29,12 @@ namespace AudioModulationEditor
 		virtual TSharedRef<SWidget> MakeControlBusListWidget();
 		virtual TSharedRef<SWidget> MakeControlBusWatchWidget();
 
+		virtual TSharedPtr<SWidget> OnConstructContextMenu() override;
+		virtual FReply OnDataRowKeyInput(const FGeometry& InGeometry, const FKeyEvent& InKeyEvent) const override;
+
+	private:
+		void BindCommands();
+
+		TSharedPtr<FUICommandList> CommandList;
 	};
 } // namespace AudioModulationEditor
diff --git a/Engine/Plugins/Runtime/AudioModulation/Source/AudioModulationEditor/Private/Insights/Views/ModulationMatrixDashboardViewFactory.cpp b/Engine/Plugins/Runtime/AudioModulation/Source/AudioModulationEditor/Private/Insights/Views/ModulationMatrixDashboardViewFactory.cpp
--- a/Engine/Plugins/Runtime/AudioModulation/Source/AudioModulationEditor/Private/Insights/Views/ModulationMatrixDashboardViewFactory.cpp
+++ b/Engine/Plugins/Runtime/AudioModulation/Source/AudioModulationEditor/Private/Insights/Views/ModulationMatrixDashboardViewFactory.cpp
@@ -4,7 +4,9 @@
 #include "AudioDeviceManager.h"
 #include "AudioInsightsStyle.h"
 #include "AudioInsightsTraceProviderBase.h"
+#include "AudioModulationEditorCommands.h"
 #include "Framework/Application/SlateApplication.h"
+#include "Framework/MultiBox/MultiBoxBuilder.h"
 #include "Fonts/FontMeasure.h"
 #include "IAudioInsightsEditorModule.h"
 #include "Internationalization/Text.h"
@@ -42,6 +44,15 @@ namespace AudioModulationEditor
 		};
 
 		CreateDefaultColumnData();
+
+		FAudioModulationEditorCommands::Register();
+
+		BindCommands();
+	}
+
+	FModulationMatrixDashboardViewFactory::~FModulationMatrixDashboardViewFactory()
+	{
+		FAudioModulationEditorCommands::Unregister();
 	}
 
 	FName FModulationMatrixDashboardViewFactory::GetName() const
@@ -101,7 +112,7 @@ namespace AudioModulationEditor
 				ModulatingSourceColumnName,
 				{
 					LOCTEXT("ModulationMatrix_ModulatingSourceDisplayName", "Modulating Source"),
-					[](const IDashboardDataViewEntry& InData) { return FText::FromString(CastEntry(InData).Name); },
+					[](const IDashboardDataViewEntry& InData) { return FText::FromString(CastEntry(InData).DisplayName); },
 					false /* bDefaultHidden */,
 					0.4f /* FillWidth */
 				}
@@ -156,6 +167,16 @@ namespace AudioModulationEditor
 		}
 	}
 
+	void FModulationMatrixDashboardViewFactory::BindCommands()
+	{
+		CommandList = MakeShared<FUICommandList>();
+
+		const FAudioModulationEditorCommands& Commands = FAudioModulationEditorCommands::Get();
+
+		CommandList->MapAction(Commands.GetBrowseCommand(), FExecuteAction::CreateLambda([this]() { BrowseToAsset(); }));
+		CommandList->MapAction(Commands.GetEditCommand(), FExecuteAction::CreateLambda([this]() { OpenAsset(); }));
+	}
+
 	TSharedRef<SWidget> FModulationMatrixDashboardViewFactory::MakeModulatingSourceTypeFilterWidget()
 	{
 		if (ModulatingSourceTypes.IsEmpty())
@@ -247,6 +268,29 @@ namespace AudioModulationEditor
 		FilterByModulatingSourceType();
 	}
 
+	TSharedPtr<SWidget> FModulationMatrixDashboardViewFactory::OnConstructContextMenu()
+	{
+		const FAudioModulationEditorCommands& Commands = FAudioModulationEditorCommands::Get();
+
+		constexpr bool bShouldCloseWindowAfterMenuSelection = true;
+		FMenuBuilder MenuBuilder(bShouldCloseWindowAfterMenuSelection, CommandList);
+
+		MenuBuilder.BeginSection("ControlBusDashboardActions", LOCTEXT("AudioInsights_ModulationMatrix_HeaderText", "Modulation Matrix Options"));
+		{
+			MenuBuilder.AddMenuEntry(Commands.GetBrowseCommand(), NAME_None, TAttribute<FText>(), TAttribute<FText>(), UE::Audio::Insights::FSlateStyle::Get().CreateIcon("AudioInsights.Icon.SoundDashboard.Browse"));
+			MenuBuilder.AddMenuEntry(Commands.GetEditCommand(), NAME_None, TAttribute<FText>(), TAttribute<FText>(), UE::Audio::Insights::FSlateStyle::Get().CreateIcon("AudioInsights.Icon.SoundDashboard.Edit"));
+		}
+
+		MenuBuilder.EndSection();
+
+		return MenuBuilder.MakeWidget();
+	}
+
+	FReply FModulationMatrixDashboardViewFactory::OnDataRowKeyInput(const FGeometry& InGeometry, const FKeyEvent& InKeyEvent) const
+	{
+		return (CommandList && CommandList->ProcessCommandBindings(InKeyEvent)) ? FReply::Handled() : FReply::Unhandled();
+	}
+
 	void FModulationMatrixDashboardViewFactory::OnControlBusesAdded(const FModulationMatrixTraceProvider::BusIdToBusInfoMap& AddedControlBuses)
 	{
 		using namespace UE::Audio::Insights;
@@ -333,7 +377,7 @@ namespace AudioModulationEditor
 		{
 			const FModulationMatrixDashboardEntry& ModulationMatrixEntry = FModulationMatrixDashboardViewFactoryPrivate::CastEntry(Entry);
 
-			return !ModulationMatrixEntry.Name.Contains(FilterString) && ModulationMatrixEntry.EntryType != EModulationMatrixEntryType::BusFinalValues;
+			return !ModulationMatrixEntry.DisplayName.Contains(FilterString) && ModulationMatrixEntry.EntryType != EModulationMatrixEntryType::BusFinalValues;
 		});
 	}
 
diff --git a/Engine/Plugins/Runtime/AudioModulation/Source/AudioModulationEditor/Private/Insights/Views/ModulationMatrixDashboardViewFactory.h b/Engine/Plugins/Runtime/AudioModulation/Source/AudioModulationEditor/Private/Insights/Views/ModulationMatrixDashboardViewFactory.h
--- a/Engine/Plugins/Runtime/AudioModulation/Source/AudioModulationEditor/Private/Insights/Views/ModulationMatrixDashboardViewFactory.h
+++ b/Engine/Plugins/Runtime/AudioModulation/Source/AudioModulationEditor/Private/Insights/Views/ModulationMatrixDashboardViewFactory.h
@@ -17,7 +17,7 @@ namespace AudioModulationEditor
 	{
 	public:
 		FModulationMatrixDashboardViewFactory();
-		virtual ~FModulationMatrixDashboardViewFactory() = default;
+		virtual ~FModulationMatrixDashboardViewFactory();
 
 		virtual FName GetName() const override;
 		virtual FText GetDisplayName() const override;
@@ -28,12 +28,17 @@ namespace AudioModulationEditor
 
 		virtual void ProcessEntries(UE::Audio::Insights::FTraceTableDashboardViewFactory::EProcessReason Reason) override;
 
+		virtual TSharedPtr<SWidget> OnConstructContextMenu() override;
+		virtual FReply OnDataRowKeyInput(const FGeometry& InGeometry, const FKeyEvent& InKeyEvent) const override;
+
 	private:
 		virtual const TMap<FName, UE::Audio::Insights::FTraceTableDashboardViewFactory::FColumnData>& GetColumns() const override;
 
 		void CreateDefaultColumnData();
 		void RegisterDelegates();
 
+		void BindCommands();
+
 		TSharedRef<SWidget> MakeModulatingSourceTypeFilterWidget();
 
 		void OnControlBusesAdded(const FModulationMatrixTraceProvider::BusIdToBusInfoMap& AddedControlBuses);
@@ -64,5 +69,7 @@ namespace AudioModulationEditor
 		using FComboboxSelectionItem = TPair<EModulatingSourceComboboxSelection, FText>;
 		TArray<TSharedPtr<FComboboxSelectionItem>> ModulatingSourceTypes;
 		TSharedPtr<FComboboxSelectionItem> SelectedModulatingSourceType;
+
+		TSharedPtr<FUICommandList> CommandList;
 	};
 } // namespace AudioModulationEditor
