[Audio Insights] Make filter rule more readable (bug hawk didn't like the inverted filter rule, logic was correct but not very readable)

diff --git a/Engine/Plugins/AudioInsights/Source/AudioInsights/Private/Views/AudioEventLogDashboardViewFactory.cpp b/Engine/Plugins/AudioInsights/Source/AudioInsights/Private/Views/AudioEventLogDashboardViewFactory.cpp
--- a/Engine/Plugins/AudioInsights/Source/AudioInsights/Private/Views/AudioEventLogDashboardViewFactory.cpp
+++ b/Engine/Plugins/AudioInsights/Source/AudioInsights/Private/Views/AudioEventLogDashboardViewFactory.cpp
@@ -197,11 +197,11 @@ namespace UE::Audio::Insights
 			{
 				if (EventLogEntry.GetDisplayName().ToString().Contains(FilterString) || EventLogEntry.ActorName.Contains(FilterString))
 				{
-					return false;
+					return true;
 				}
 			}
 
-			return true;
+			return false;
 		});
 	}
 
diff --git a/Engine/Plugins/AudioInsights/Source/AudioInsights/Private/Views/SubmixDashboardViewFactory.cpp b/Engine/Plugins/AudioInsights/Source/AudioInsights/Private/Views/SubmixDashboardViewFactory.cpp
--- a/Engine/Plugins/AudioInsights/Source/AudioInsights/Private/Views/SubmixDashboardViewFactory.cpp
+++ b/Engine/Plugins/AudioInsights/Source/AudioInsights/Private/Views/SubmixDashboardViewFactory.cpp
@@ -150,7 +150,7 @@ namespace UE::Audio::Insights
 		{
 			const FSubmixDashboardEntry& SubmixEntry = static_cast<const FSubmixDashboardEntry&>(Entry);
 			
-			return !SubmixEntry.GetDisplayName().ToString().Contains(FilterString);
+			return SubmixEntry.GetDisplayName().ToString().Contains(FilterString);
 		});
 	}
 
diff --git a/Engine/Plugins/AudioInsights/Source/AudioInsights/Private/Views/VirtualLoopDashboardViewFactory.cpp b/Engine/Plugins/AudioInsights/Source/AudioInsights/Private/Views/VirtualLoopDashboardViewFactory.cpp
--- a/Engine/Plugins/AudioInsights/Source/AudioInsights/Private/Views/VirtualLoopDashboardViewFactory.cpp
+++ b/Engine/Plugins/AudioInsights/Source/AudioInsights/Private/Views/VirtualLoopDashboardViewFactory.cpp
@@ -56,12 +56,8 @@ namespace UE::Audio::Insights
 		FTraceTableDashboardViewFactory::FilterEntries<FVirtualLoopTraceProvider>([&FilterString](const IDashboardDataViewEntry& Entry)
 		{
 			const FVirtualLoopDashboardEntry& VirtualLoopEntry = static_cast<const FVirtualLoopDashboardEntry&>(Entry);
-			if (VirtualLoopEntry.GetDisplayName().ToString().Contains(FilterString))
-			{
-				return false;
-			}
 
-			return true;
+			return VirtualLoopEntry.GetDisplayName().ToString().Contains(FilterString);
 		});
 	}
 
diff --git a/Engine/Plugins/AudioInsights/Source/AudioInsights/Public/Views/TableDashboardViewFactory.h b/Engine/Plugins/AudioInsights/Source/AudioInsights/Public/Views/TableDashboardViewFactory.h
--- a/Engine/Plugins/AudioInsights/Source/AudioInsights/Public/Views/TableDashboardViewFactory.h
+++ b/Engine/Plugins/AudioInsights/Source/AudioInsights/Public/Views/TableDashboardViewFactory.h
@@ -100,7 +100,7 @@ namespace UE::Audio::Insights
 					{
 						auto FilterEntry = [this, &IsFiltered](const typename TableProviderType::FEntryPair& Pair)
 						{
-							return !IsFiltered((const IDashboardDataViewEntry&)(*Pair.Value));
+							return IsFiltered((const IDashboardDataViewEntry&)(*Pair.Value));
 						};
 						Algo::TransformIf(*DeviceData, DataViewEntries, FilterEntry, TransformEntry);
 					}
diff --git a/Engine/Plugins/AudioInsights/Source/AudioInsightsEditor/Private/Views/AudioBusesDashboardViewFactory.cpp b/Engine/Plugins/AudioInsights/Source/AudioInsightsEditor/Private/Views/AudioBusesDashboardViewFactory.cpp
--- a/Engine/Plugins/AudioInsights/Source/AudioInsightsEditor/Private/Views/AudioBusesDashboardViewFactory.cpp
+++ b/Engine/Plugins/AudioInsights/Source/AudioInsightsEditor/Private/Views/AudioBusesDashboardViewFactory.cpp
@@ -400,7 +400,7 @@ namespace UE::Audio::Insights
 		{
 			const FAudioBusAssetDashboardEntry& AudioBusEntry = static_cast<const FAudioBusAssetDashboardEntry&>(Entry);
 			
-			return !AudioBusEntry.GetDisplayName().ToString().Contains(FilterString);
+			return AudioBusEntry.GetDisplayName().ToString().Contains(FilterString);
 		});
 	}
 
diff --git a/Engine/Plugins/AudioInsights/Source/AudioInsightsEditor/Private/Views/SubmixesDashboardViewFactory.cpp b/Engine/Plugins/AudioInsights/Source/AudioInsightsEditor/Private/Views/SubmixesDashboardViewFactory.cpp
--- a/Engine/Plugins/AudioInsights/Source/AudioInsightsEditor/Private/Views/SubmixesDashboardViewFactory.cpp
+++ b/Engine/Plugins/AudioInsights/Source/AudioInsightsEditor/Private/Views/SubmixesDashboardViewFactory.cpp
@@ -169,7 +169,7 @@ namespace UE::Audio::Insights
 		{
 			const FSoundSubmixAssetDashboardEntry& SubmixEntry = static_cast<const FSoundSubmixAssetDashboardEntry&>(Entry);
 			
-			return !SubmixEntry.GetDisplayName().ToString().Contains(FilterString);
+			return SubmixEntry.GetDisplayName().ToString().Contains(FilterString);
 		});
 	}
 
diff --git a/Engine/Plugins/Runtime/AudioModulation/Source/AudioModulationEditor/Private/Insights/Views/ControlBusDashboardViewFactory.cpp b/Engine/Plugins/Runtime/AudioModulation/Source/AudioModulationEditor/Private/Insights/Views/ControlBusDashboardViewFactory.cpp
--- a/Engine/Plugins/Runtime/AudioModulation/Source/AudioModulationEditor/Private/Insights/Views/ControlBusDashboardViewFactory.cpp
+++ b/Engine/Plugins/Runtime/AudioModulation/Source/AudioModulationEditor/Private/Insights/Views/ControlBusDashboardViewFactory.cpp
@@ -71,12 +71,8 @@ namespace AudioModulationEditor
 		FTraceTableDashboardViewFactory::FilterEntries<FControlBusTraceProvider>([&FilterString](const IDashboardDataViewEntry& Entry)
 		{
 			const FControlBusDashboardEntry& ControlBusEntry = ControlBusPrivate::CastEntry(Entry);
-			if (ControlBusEntry.DisplayName.Contains(FilterString))
-			{
-				return false;
-			}
 
-			return true;
+			return ControlBusEntry.DisplayName.Contains(FilterString);
 		});
 	}
 
diff --git a/Engine/Plugins/Runtime/AudioModulation/Source/AudioModulationEditor/Private/Insights/Views/ModulationMatrixDashboardViewFactory.cpp b/Engine/Plugins/Runtime/AudioModulation/Source/AudioModulationEditor/Private/Insights/Views/ModulationMatrixDashboardViewFactory.cpp
--- a/Engine/Plugins/Runtime/AudioModulation/Source/AudioModulationEditor/Private/Insights/Views/ModulationMatrixDashboardViewFactory.cpp
+++ b/Engine/Plugins/Runtime/AudioModulation/Source/AudioModulationEditor/Private/Insights/Views/ModulationMatrixDashboardViewFactory.cpp
@@ -381,7 +381,7 @@ namespace AudioModulationEditor
 		{
 			const FModulationMatrixDashboardEntry& ModulationMatrixEntry = FModulationMatrixDashboardViewFactoryPrivate::CastEntry(Entry);
 
-			return !ModulationMatrixEntry.DisplayName.Contains(FilterString) && ModulationMatrixEntry.EntryType != EModulationMatrixEntryType::BusFinalValues;
+			return ModulationMatrixEntry.DisplayName.Contains(FilterString) || ModulationMatrixEntry.EntryType == EModulationMatrixEntryType::BusFinalValues;
 		});
 	}
 
