Active Control Bus Mixes debug display

diff --git a/Engine/Plugins/Runtime/AudioModulation/Source/AudioModulation/Private/AudioModulation.cpp b/Engine/Plugins/Runtime/AudioModulation/Source/AudioModulation/Private/AudioModulation.cpp
--- a/Engine/Plugins/Runtime/AudioModulation/Source/AudioModulation/Private/AudioModulation.cpp
+++ b/Engine/Plugins/Runtime/AudioModulation/Source/AudioModulation/Private/AudioModulation.cpp
@@ -133,6 +133,16 @@ namespace AudioModulation
 		ModSystem->SetDebugGeneratorTypeFilter(InFilter, bInIsEnabled);
 	}
 
+	void FAudioModulationManager::SetDebugActiveMixesEnabled(bool bInIsEnabled)
+	{
+		ModSystem->SetDebugActiveMixesEnabled(bInIsEnabled);
+	}
+
+	void FAudioModulationManager::SetDebugActiveGlobalMixesEnabled(bool bInIsEnabled)
+	{
+		ModSystem->SetDebugActiveGlobalMixesEnabled(bInIsEnabled);
+	}
+
 	void FAudioModulationManager::SetDebugGeneratorsEnabled(bool bInIsEnabled)
 	{
 		ModSystem->SetDebugGeneratorsEnabled(bInIsEnabled);
diff --git a/Engine/Plugins/Runtime/AudioModulation/Source/AudioModulation/Private/AudioModulationDebugger.cpp b/Engine/Plugins/Runtime/AudioModulation/Source/AudioModulation/Private/AudioModulationDebugger.cpp
--- a/Engine/Plugins/Runtime/AudioModulation/Source/AudioModulation/Private/AudioModulationDebugger.cpp
+++ b/Engine/Plugins/Runtime/AudioModulation/Source/AudioModulation/Private/AudioModulationDebugger.cpp
@@ -160,6 +160,41 @@ static FAutoConsoleCommandWithWorldAndArgs CVarAudioModulationDebugEnableGenerat
 		});
 	}));
 
+
+static FAutoConsoleCommandWithWorldAndArgs CVarAudioModulationDebugActiveMixes(
+	TEXT("au.Debug.Modulation.ActiveMixes"),
+	TEXT("Whether or not to enable displaying active Control Bus Mixes. Arguments:\n"
+		"Enabled (Default: 1/True) - Whether or not to have generators view enabled.\n"),
+	FConsoleCommandWithWorldAndArgsDelegate::CreateLambda([](const TArray<FString>& Args, UWorld* World)
+	{
+		using namespace AudioModulation;
+		Debug::UpdateObjectFilter(Args, World, [](FAudioModulationManager& InModulation, const FString* InFilter)
+		{
+			if (InFilter)
+			{
+				const bool bEnabled = InFilter->ToBool();
+				InModulation.SetDebugActiveMixesEnabled(bEnabled);
+			}
+		});
+	}));
+
+static FAutoConsoleCommandWithWorldAndArgs CVarAudioModulationDebugShowGlobalMixes(
+	TEXT("au.Debug.Modulation.ActiveMixes.ShowGlobalMixes"),
+	TEXT("Whether or not to show Global Mixes when displaying active mixes using 'au.Debug.Modulation.ActiveMixes'. Arguments:\n"
+		"Enabled (Default: 1/True) - Whether or not to have generators view enabled.\n"),
+	FConsoleCommandWithWorldAndArgsDelegate::CreateLambda([](const TArray<FString>& Args, UWorld* World)
+	{
+		using namespace AudioModulation;
+		Debug::UpdateObjectFilter(Args, World, [](FAudioModulationManager& InModulation, const FString* InFilter)
+		{
+			if (InFilter)
+			{
+				const bool bEnabled = InFilter->ToBool();
+				InModulation.SetDebugActiveGlobalMixesEnabled(bEnabled);
+			}
+		});
+	}));
+
 namespace AudioModulation
 {
 	namespace Debug
@@ -366,6 +401,47 @@ namespace AudioModulation
 			return Y + Height;
 		}
 
+		int32 RenderStatActiveMixes(const TArray<FControlBusMixDebugInfo>& FilteredMixes, bool bShowGlobalMixes, FCanvas& Canvas, int32 X, int32 Y, const UFont& Font)
+		{
+			FAudioDevice* AudioDevice = nullptr;
+
+			int32 Height = 12;
+			int32 Width = 12;
+			Font.GetStringHeightAndWidth(TEXT("@"), Height, Width);
+			Canvas.DrawShadowedString(X, Y, TEXT("Active Control Bus Mixes:"), &Font, FColor::Green);
+			Y += Height;
+			X += XIndent;
+			
+			bool bDisplayedBusMixes = false;
+
+			if (FilteredMixes.Num() > 0)
+			{
+			 	bDisplayedBusMixes = true;
+
+				for (const FControlBusMixDebugInfo& DebugMix : FilteredMixes)
+				{
+					if (!bShowGlobalMixes && DebugMix.bIsGlobal)
+					{
+						continue;
+					}
+					
+					const FString BusMixInfo = FString::Printf(TEXT("%s - Timer: %.3f - Total Ref Count: %i"), *DebugMix.Name, DebugMix.Timer, DebugMix.RefCount);
+					const FColor& TextColor = (FColor::White);
+					Canvas.DrawShadowedString(X, Y, *BusMixInfo, &Font, TextColor);
+					Y += Height;
+				}
+			}
+
+			if (!bDisplayedBusMixes)
+			{
+				Canvas.DrawShadowedString(X, Y, TEXT("None"), &Font, FColor::White);
+				Y += Height;
+			}
+			
+			Y += Height;
+			return Y;
+		}
+		
 		template <typename T>
 		bool CompareNames(const T& A, const T& B)
 		{
@@ -391,6 +467,26 @@ namespace AudioModulation
 				}
 			}
 		}
+
+		template <typename T, typename U, typename K>
+		void FilterDebugMap(const T& Map, const FString& FilterString, int32 MaxCount, TMap<K, const U*>& OutFilteredMap)
+		{
+			int32 FilteredItemCount = 0;
+			for (const TPair<K, U>& IdItemPair : Map)
+			{
+				const U& Item = IdItemPair.Value;
+				const bool Filtered = !FilterString.IsEmpty()
+					&& !Item.GetName().ToString().Contains(FilterString);
+				if (!Filtered)
+				{
+					OutFilteredMap.Add(IdItemPair.Key, &Item);
+					if (++FilteredItemCount >= MaxCount)
+					{
+						return;
+					}
+				}
+			}
+		}
 	} // namespace Debug
 
 	FAudioModulationDebugger::FAudioModulationDebugger()
@@ -398,11 +494,13 @@ namespace AudioModulation
 		, bShowRenderStatMix(1)
 		, bShowGenerators(0)
 		, bEnableAllGenerators(0)
+		, bShowActiveMixes(0)
+		, bShowGlobalMixes(1)
 		, ElapsedSinceLastUpdate(0.0f)
 	{
 	}
 
-	void FAudioModulationDebugger::UpdateDebugData(double InElapsed, const FReferencedProxies& InRefProxies)
+	void FAudioModulationDebugger::UpdateDebugData(double InElapsed, const FReferencedProxies& InRefProxies, Audio::FDeviceId DeviceId)
 	{
 		if (!bActive)
 		{
@@ -434,25 +532,54 @@ namespace AudioModulation
 			RefreshedFilteredBuses.Add(DebugInfo);
 		}
 
-		static const int32 MaxFilteredMixes = 16;
-		TArray<const FModulatorBusMixProxy*> FilteredMixProxies;
-		Debug::FilterDebugArray<FBusMixProxyMap, FModulatorBusMixProxy>(InRefProxies.BusMixes, MixStringFilter, MaxFilteredMixes, FilteredMixProxies);
-		TArray<FControlBusMixDebugInfo> RefreshedFilteredMixes;
-		for (const FModulatorBusMixProxy* Proxy : FilteredMixProxies)
+		FAudioModulationSystem* ModulationSystem = nullptr;
+		if (FAudioDeviceManager* DeviceManager = FAudioDeviceManager::Get())
 		{
-			FControlBusMixDebugInfo DebugInfo;
-			DebugInfo.Name = Proxy->GetName().ToString();
-			DebugInfo.RefCount = Proxy->GetRefCount();
-			for (const TPair<FBusId, FModulatorBusMixStageProxy>& Stage : Proxy->Stages)
+			if (FAudioDeviceHandle DeviceHandle = DeviceManager->GetAudioDevice(DeviceId))
 			{
-				FControlBusMixStageDebugInfo StageDebugInfo;
-				StageDebugInfo.CurrentValue = Stage.Value.Value.GetCurrentValue();
-				StageDebugInfo.TargetValue = Stage.Value.Value.TargetValue;
-				DebugInfo.Stages.Add(Stage.Key, StageDebugInfo);
+				if (IAudioModulationManager* Modulation = DeviceHandle->ModulationInterface.Get())
+				{
+					if (FAudioModulationManager* ModulationManager = static_cast<FAudioModulationManager*>(Modulation))
+					{
+						ModulationSystem = &ModulationManager->GetSystem();
+					}
+				}
 			}
-			RefreshedFilteredMixes.Add(DebugInfo);
 		}
 
+		TArray<FControlBusMixDebugInfo> RefreshedFilteredMixes;
+		if (ModulationSystem)
+		{
+			static const int32 MaxFilteredMixes = 16;
+			TMap<FBusMixId, const FModulatorBusMixProxy*> FilteredMixProxies;
+			Debug::FilterDebugMap<FBusMixProxyMap, FModulatorBusMixProxy, FBusMixId>(InRefProxies.BusMixes, MixStringFilter, MaxFilteredMixes, FilteredMixProxies);
+			for (TPair<FBusMixId, const FModulatorBusMixProxy*>& Pair : FilteredMixProxies)
+			{
+				const FModulatorBusMixProxy* Proxy = Pair.Value;
+			
+				FControlBusMixDebugInfo DebugInfo;
+				DebugInfo.Name = Proxy->GetName().ToString();
+				DebugInfo.RefCount = Proxy->GetRefCount();
+				DebugInfo.Id = Pair.Key;
+				DebugInfo.Timer = Proxy->GetTimeRemaining();
+				// Bus mix is global if it only has one stage and its ID is the Control Bus's ID.
+				if (Proxy->Stages.Num() == 1)
+				{
+					uint32 BusId = Proxy->Stages.begin().Key();
+					DebugInfo.bIsGlobal = ModulationSystem->IsBusMixGlobal(BusId);
+				}
+			
+				for (const TPair<FBusId, FModulatorBusMixStageProxy>& Stage : Proxy->Stages)
+				{
+					FControlBusMixStageDebugInfo StageDebugInfo;
+					StageDebugInfo.CurrentValue = Stage.Value.Value.GetCurrentValue();
+					StageDebugInfo.TargetValue = Stage.Value.Value.TargetValue;
+					DebugInfo.Stages.Add(Stage.Key, StageDebugInfo);
+				}
+				RefreshedFilteredMixes.Add(DebugInfo);
+			}
+		}
+		
 		static const int32 MaxFilteredGenerators = 8;
 		TArray<const FModulatorGeneratorProxy*> FilteredGeneratorProxies;
 		Debug::FilterDebugArray<FGeneratorProxyMap, FModulatorGeneratorProxy>(InRefProxies.Generators, GeneratorStringFilter, MaxFilteredGenerators, FilteredGeneratorProxies);
@@ -599,6 +726,17 @@ namespace AudioModulation
 		}
 	}
 
+	void FAudioModulationDebugger::SetDebugActiveMixesEnabled(bool bInIsEnabled)
+	{
+		bShowActiveMixes = bInIsEnabled;
+	}
+
+	void FAudioModulationDebugger::SetDebugActiveGlobalMixesEnabled(bool bInIsEnabled)
+	{
+		bShowGlobalMixes = bInIsEnabled;
+		bShowActiveMixes = 1;
+	}
+
 	int32 FAudioModulationDebugger::OnRenderStat(FCanvas& Canvas, int32 X, int32 Y, const UFont& Font)
 	{
 		check(IsInGameThread());
@@ -643,6 +781,11 @@ namespace AudioModulation
 			}
 		}
 
+		if (bShowActiveMixes)
+		{
+			Y = Debug::RenderStatActiveMixes(FilteredMixes, bShowGlobalMixes, Canvas, X, Y, Font);
+		}
+
 		return Y;
 	}
 
diff --git a/Engine/Plugins/Runtime/AudioModulation/Source/AudioModulation/Private/AudioModulationDebugger.h b/Engine/Plugins/Runtime/AudioModulation/Source/AudioModulation/Private/AudioModulationDebugger.h
--- a/Engine/Plugins/Runtime/AudioModulation/Source/AudioModulation/Private/AudioModulationDebugger.h
+++ b/Engine/Plugins/Runtime/AudioModulation/Source/AudioModulation/Private/AudioModulationDebugger.h
@@ -26,8 +26,10 @@ namespace AudioModulation
 	struct FControlBusMixDebugInfo
 	{
 		FString Name;
-		uint32 Id;
-		uint32 RefCount;
+		uint32 Id = -1;
+		uint32 RefCount = 0;
+		bool bIsGlobal = false;
+		double Timer = -1.0;
 
 		TMap<uint32, FControlBusMixStageDebugInfo> Stages;
 	};
@@ -76,13 +78,15 @@ namespace AudioModulation
 	public:
 		FAudioModulationDebugger();
 
-		void UpdateDebugData(double InElapsed, const FReferencedProxies& InRefProxies);
+		void UpdateDebugData(double InElapsed, const FReferencedProxies& InRefProxies, Audio::FDeviceId DeviceId);
 		void SetDebugBusFilter(const FString* InNameFilter);
 		void SetDebugMatrixEnabled(bool bInIsEnabled);
 		void SetDebugMixFilter(const FString* InNameFilter);
 		void SetDebugGeneratorsEnabled(bool bInIsEnabled);
 		void SetDebugGeneratorFilter(const FString* InFilter);
 		void SetDebugGeneratorTypeFilter(const FString* InFilter, bool bInIsEnabled);
+		void SetDebugActiveMixesEnabled(bool bInIsEnabled);
+		void SetDebugActiveGlobalMixesEnabled(bool bInIsEnabled);
 		bool OnPostHelp(FCommonViewportClient& ViewportClient, const TCHAR* Stream);
 		int32 OnRenderStat(FCanvas& Canvas, int32 X, int32 Y, const UFont& Font);
 		bool OnToggleStat(FCommonViewportClient& ViewportClient, const TCHAR* Stream);
@@ -94,6 +98,8 @@ namespace AudioModulation
 		uint8 bShowRenderStatMix : 1;
 		uint8 bShowGenerators : 1;
 		uint8 bEnableAllGenerators : 1;
+		uint8 bShowActiveMixes : 1;
+		uint8 bShowGlobalMixes : 1;
 
 		TArray<FControlBusDebugInfo> FilteredBuses;
 		TArray<FControlBusMixDebugInfo> FilteredMixes;
diff --git a/Engine/Plugins/Runtime/AudioModulation/Source/AudioModulation/Private/AudioModulationSystem.cpp b/Engine/Plugins/Runtime/AudioModulation/Source/AudioModulation/Private/AudioModulationSystem.cpp
--- a/Engine/Plugins/Runtime/AudioModulation/Source/AudioModulation/Private/AudioModulationSystem.cpp
+++ b/Engine/Plugins/Runtime/AudioModulation/Source/AudioModulation/Private/AudioModulationSystem.cpp
@@ -323,6 +323,16 @@ namespace AudioModulation
 		Debugger->SetDebugGeneratorTypeFilter(InFilter, bInEnabled);
 	}
 
+	void FAudioModulationSystem::SetDebugActiveMixesEnabled(bool bInIsEnabled)
+	{
+		Debugger->SetDebugActiveMixesEnabled(bInIsEnabled);
+	}
+
+	void FAudioModulationSystem::SetDebugActiveGlobalMixesEnabled(bool bInIsEnabled)
+	{
+		Debugger->SetDebugActiveGlobalMixesEnabled(bInIsEnabled);
+	}
+
 	void FAudioModulationSystem::SetDebugMatrixEnabled(bool bInIsEnabled)
 	{
 		Debugger->SetDebugMatrixEnabled(bInIsEnabled);
@@ -399,6 +409,12 @@ namespace AudioModulation
 		ActiveGlobalBusValueMixes.Reset();
 	}
 
+	bool FAudioModulationSystem::IsBusMixGlobal(FBusMixId InBusMixId) const
+	{
+		// Right now, anyone calling this is assuming the ID of ActiveGlobalBuxValueMixes is the Control Bus ID, not the Control Bus Mix ID.
+		return ActiveGlobalBusValueMixes.Contains(InBusMixId);
+	}
+
 	USoundControlBusMix* FAudioModulationSystem::CreateBusMixFromValue(FName Name, const TArray<USoundControlBus*>& Buses, float Value, float AttackTime, float ReleaseTime)
 	{
 		if (TObjectPtr<USoundControlBusMix> NewGlobalMix = NewObject<USoundControlBusMix>(GetTransientPackage(), Name))
@@ -720,7 +736,7 @@ namespace AudioModulation
 #endif // UE_AUDIO_PROFILERTRACE_ENABLED
 
 #if !UE_BUILD_SHIPPING
- 		Debugger->UpdateDebugData(InElapsed, RefProxies);
+ 		Debugger->UpdateDebugData(InElapsed, RefProxies, AudioDeviceId);
 #endif // !UE_BUILD_SHIPPING
 	}
 
diff --git a/Engine/Plugins/Runtime/AudioModulation/Source/AudioModulation/Private/AudioModulationSystem.h b/Engine/Plugins/Runtime/AudioModulation/Source/AudioModulation/Private/AudioModulationSystem.h
--- a/Engine/Plugins/Runtime/AudioModulation/Source/AudioModulation/Private/AudioModulationSystem.h
+++ b/Engine/Plugins/Runtime/AudioModulation/Source/AudioModulation/Private/AudioModulationSystem.h
@@ -160,6 +160,9 @@ namespace AudioModulation
 		/* Clears all global bus mix values over the prescribed FadeTime. If FadeTime is non-positive, returns to the bus's respective parameter default immediately. */
 		void ClearAllGlobalBusMixValues(float FadeTime = -1.0f);
 
+		/* Returns whether the given bus Mix is a Global Bus Mix, created using the SetGlobalBusMixValue function. */
+		bool IsBusMixGlobal(FBusMixId InBusMixId) const;
+		
 		/* Create a mix from an array of buses where a supplied default value and associated timing parameters are applied for each bus's stage */
 		USoundControlBusMix* CreateBusMixFromValue(FName Name, const TArray<USoundControlBus*>& Buses, float Value, float AttackTime = -1.0f, float ReleaseTime = -1.0f);
 
@@ -255,6 +258,7 @@ namespace AudioModulation
 		TSet<FGeneratorHandle> ManuallyActivatedGenerators;
 
 		// Global mixes each containing a single stage for any globally manipulated bus stage value.
+		// the key in this map is the Bus's unique ID, not the Control Bus Mix's.
 		TMap<uint32, TObjectPtr<USoundControlBusMix>> ActiveGlobalBusValueMixes;
 
 		// Command queue to be consumed on processing thread 
@@ -280,6 +284,8 @@ namespace AudioModulation
 		void SetDebugGeneratorsEnabled(bool bInIsEnabled);
 		void SetDebugGeneratorFilter(const FString* InFilter);
 		void SetDebugGeneratorTypeFilter(const FString* InFilter, bool bInIsEnabled);
+		virtual void SetDebugActiveMixesEnabled(bool bInIsEnabled);
+		virtual void SetDebugActiveGlobalMixesEnabled(bool bInIsEnabled);
 		bool OnPostHelp(FCommonViewportClient* ViewportClient, const TCHAR* Stream);
 		int OnRenderStat(FViewport* Viewport, FCanvas* Canvas, int32 X, int32 Y, const UFont& Font, const FVector* ViewLocation, const FRotator* ViewRotation);
 		bool OnToggleStat(FCommonViewportClient* ViewportClient, const TCHAR* Stream);
diff --git a/Engine/Plugins/Runtime/AudioModulation/Source/AudioModulation/Private/SoundControlBusMixProxy.cpp b/Engine/Plugins/Runtime/AudioModulation/Source/AudioModulation/Private/SoundControlBusMixProxy.cpp
--- a/Engine/Plugins/Runtime/AudioModulation/Source/AudioModulation/Private/SoundControlBusMixProxy.cpp
+++ b/Engine/Plugins/Runtime/AudioModulation/Source/AudioModulation/Private/SoundControlBusMixProxy.cpp
@@ -279,6 +279,15 @@ namespace AudioModulation
 #endif
 		}
 	}
+
+	double FModulatorBusMixProxy::GetTimeRemaining() const
+	{
+		if (Duration >= 0)
+		{
+			return TimeRemaining;
+		}
+		return -1.0;
+	}
 } // namespace AudioModulation
 
 #undef LOCTEXT_NAMESPACE // AudioModulation
diff --git a/Engine/Plugins/Runtime/AudioModulation/Source/AudioModulation/Private/SoundControlBusMixProxy.h b/Engine/Plugins/Runtime/AudioModulation/Source/AudioModulation/Private/SoundControlBusMixProxy.h
--- a/Engine/Plugins/Runtime/AudioModulation/Source/AudioModulation/Private/SoundControlBusMixProxy.h
+++ b/Engine/Plugins/Runtime/AudioModulation/Source/AudioModulation/Private/SoundControlBusMixProxy.h
@@ -89,6 +89,9 @@ namespace AudioModulation
 
 		void Update(const double Elapsed, FBusProxyMap& ProxyMap);
 
+		// If the Control Bus Mix is using a timer, returns the amount of time left before deactivation (seconds). Otherwise returns -1.
+		double GetTimeRemaining() const;
+		
 		using FStageMap = TMap<FBusId, FModulatorBusMixStageProxy>;
 		FStageMap Stages;
 
diff --git a/Engine/Plugins/Runtime/AudioModulation/Source/AudioModulation/Public/AudioModulation.h b/Engine/Plugins/Runtime/AudioModulation/Source/AudioModulation/Public/AudioModulation.h
--- a/Engine/Plugins/Runtime/AudioModulation/Source/AudioModulation/Public/AudioModulation.h
+++ b/Engine/Plugins/Runtime/AudioModulation/Source/AudioModulation/Public/AudioModulation.h
@@ -44,6 +44,8 @@ namespace AudioModulation
 		UE_API virtual void SetDebugGeneratorsEnabled(bool bInIsEnabled);
 		UE_API virtual void SetDebugGeneratorFilter(const FString* InFilter);
 		UE_API virtual void SetDebugGeneratorTypeFilter(const FString* InFilter, bool bInIsEnabled);
+		UE_API virtual void SetDebugActiveMixesEnabled(bool bInIsEnabled);
+		UE_API virtual void SetDebugActiveGlobalMixesEnabled(bool bInIsEnabled);
 
 		UE_API virtual bool OnPostHelp(FCommonViewportClient* ViewportClient, const TCHAR* Stream) override;
 		UE_API virtual int32 OnRenderStat(FViewport* Viewport, FCanvas* Canvas, int32 X, int32 Y, const UFont& Font, const FVector* ViewLocation, const FRotator* ViewRotation) override;
