[Audio Insights] [AudioModulation] Fix for Modulation Matrix showing the level name as control bus mix name when creating a control bus mix via Blueprints

diff --git a/Engine/Plugins/Runtime/AudioModulation/Source/AudioModulation/Private/SoundModulationProxy.h b/Engine/Plugins/Runtime/AudioModulation/Source/AudioModulation/Private/SoundModulationProxy.h
--- a/Engine/Plugins/Runtime/AudioModulation/Source/AudioModulation/Private/SoundModulationProxy.h
+++ b/Engine/Plugins/Runtime/AudioModulation/Source/AudioModulation/Private/SoundModulationProxy.h
@@ -234,6 +234,29 @@ namespace AudioModulation
 	class TModulatorBase
 	{
 	private:
+#if !UE_BUILD_SHIPPING
+		FString ResolveName(const FString& InNamePath)
+		{
+			const FSoftObjectPath SoftObjectPath = FSoftObjectPath(InNamePath);
+
+			if (SoftObjectPath.IsAsset())
+			{
+				return SoftObjectPath.GetAssetName();
+			}
+			else if (SoftObjectPath.IsSubobject())
+			{
+				FString Left;
+				FString Right;
+				if (SoftObjectPath.GetSubPathString().Split(".", &Left, &Right, ESearchCase::IgnoreCase, ESearchDir::FromEnd))
+				{
+					return Right;
+				}
+			}
+
+			return InNamePath;
+		}
+#endif // !UE_BUILD_SHIPPING
+
 		IdType Id;
 
 #if !UE_BUILD_SHIPPING
@@ -250,7 +273,7 @@ namespace AudioModulation
 		TModulatorBase(const FName InPath, const uint32 InId)
 			: Id(static_cast<IdType>(InId))
 #if !UE_BUILD_SHIPPING
-			, Name(FSoftObjectPath(InPath.ToString()).GetAssetName())
+			, Name(ResolveName(InPath.ToString()))
 			, Path(InPath)
 #endif // !UE_BUILD_SHIPPING
 		{
@@ -259,7 +282,7 @@ namespace AudioModulation
 		TModulatorBase(const FString& InPath, const uint32 InId)
 			: Id(static_cast<IdType>(InId))
 #if !UE_BUILD_SHIPPING
-			, Name(FSoftObjectPath(InPath).GetAssetName())
+			, Name(ResolveName(InPath))
 			, Path(InPath)
 #endif // !UE_BUILD_SHIPPING
 		{
diff --git a/Engine/Plugins/Runtime/AudioModulation/Source/AudioModulationEditor/Private/Insights/Providers/ModulationMatrixTraceProvider.cpp b/Engine/Plugins/Runtime/AudioModulation/Source/AudioModulationEditor/Private/Insights/Providers/ModulationMatrixTraceProvider.cpp
--- a/Engine/Plugins/Runtime/AudioModulation/Source/AudioModulationEditor/Private/Insights/Providers/ModulationMatrixTraceProvider.cpp
+++ b/Engine/Plugins/Runtime/AudioModulation/Source/AudioModulationEditor/Private/Insights/Providers/ModulationMatrixTraceProvider.cpp
@@ -7,6 +7,30 @@
 
 namespace AudioModulationEditor
 {
+	namespace FModulationMatrixTraceProviderPrivate
+	{
+		FString ResolveDisplayName(const FString& InNamePath)
+		{
+			const FSoftObjectPath SoftObjectPath = FSoftObjectPath(InNamePath);
+
+			if (SoftObjectPath.IsAsset())
+			{
+				return SoftObjectPath.GetAssetName();
+			}
+			else if (SoftObjectPath.IsSubobject())
+			{
+				FString Left;
+				FString Right;
+				if (SoftObjectPath.GetSubPathString().Split(".", &Left, &Right, ESearchCase::IgnoreCase, ESearchDir::FromEnd))
+				{
+					return Right;
+				}
+			}
+
+			return InNamePath;
+		}
+	}
+
 	FModulationMatrixTraceProvider::FModulationMatrixTraceProvider()
 		: UE::Audio::Insights::TDeviceDataMapTraceProvider<uint32, TSharedPtr<FModulationMatrixDashboardEntry>>(GetName_Static())
 	{
@@ -72,6 +96,8 @@ namespace AudioModulationEditor
 
 	bool FModulationMatrixTraceProvider::ProcessMessages()
 	{
+		using namespace FModulationMatrixTraceProviderPrivate;
+
 		// Helper lambdas
 		auto BumpEntryFunc = [this](const FModulationMatrixMessageBase& Msg)
 		{
@@ -127,13 +153,11 @@ namespace AudioModulationEditor
 			{
 				if (TMap<FBusId, FString>* FoundBusInfoPtr = ModulatingSourceIdToBusesInfoMap.Find(Msg.ModulatingSourceId))
 				{
-					FString AssetName = FSoftObjectPath(Msg.BusName).GetAssetName();
-					FoundBusInfoPtr->FindOrAdd(Msg.SourceId, AssetName.IsEmpty() ? Msg.BusName : MoveTemp(AssetName));
+					FoundBusInfoPtr->FindOrAdd(Msg.SourceId, ResolveDisplayName(Msg.BusName));
 				}
 				else
 				{
-					FString AssetName = FSoftObjectPath(Msg.BusName).GetAssetName();
-					ModulatingSourceIdToBusesInfoMap.Emplace(Msg.ModulatingSourceId).Emplace(Msg.SourceId, AssetName.IsEmpty() ? Msg.BusName : MoveTemp(AssetName));
+					ModulatingSourceIdToBusesInfoMap.Emplace(Msg.ModulatingSourceId).Emplace(Msg.SourceId, ResolveDisplayName(Msg.BusName));
 				}
 			});
 
@@ -142,7 +166,7 @@ namespace AudioModulationEditor
 			{
 				FModulationMatrixDashboardEntry& EntryRef = *OutEntry->Get();
 				EntryRef.Name = *Msg.Name;
-				EntryRef.DisplayName = EntryRef.GetDisplayName().ToString();
+				EntryRef.DisplayName = ResolveDisplayName(Msg.Name);
 				EntryRef.EntryType = EModulationMatrixEntryType::BusMix;
 
 				UpdateActiveModulatorSourceIds(Msg.DeviceId, EntryRef.SourceId);
