[Audio Insights] [AudioModulation] Added ResolveObjectDisplayName function to AudioInsightUtils + fix for Control Bus dashboard showing the level name as control bus name when creating a control bus via Blueprints

diff --git a/Engine/Plugins/AudioInsights/Source/AudioInsights/Public/AudioInsightsUtils.h b/Engine/Plugins/AudioInsights/Source/AudioInsights/Public/AudioInsightsUtils.h
--- a/Engine/Plugins/AudioInsights/Source/AudioInsights/Public/AudioInsightsUtils.h
+++ b/Engine/Plugins/AudioInsights/Source/AudioInsights/Public/AudioInsightsUtils.h
@@ -25,5 +25,26 @@ namespace UE::Audio::Insights
 
 			return FText::AsNumber(::Audio::ConvertToDecibels(InAmplitude), FSlateStyle::Get().GetAmpFloatFormat());
 		}
+
+		static FString ResolveObjectDisplayName(const FString& InObjectPath)
+		{
+			const FSoftObjectPath SoftObjectPath = FSoftObjectPath(InObjectPath);
+
+			if (SoftObjectPath.IsAsset())
+			{
+				return SoftObjectPath.GetAssetName();
+			}
+			else if (SoftObjectPath.IsSubobject())
+			{
+				FString Left;
+				FString Right;
+				if (SoftObjectPath.GetSubPathString().Split(".", &Left, &Right, ESearchCase::IgnoreCase, ESearchDir::FromEnd))
+				{
+					return Right;
+				}
+			}
+
+			return InObjectPath;
+		}
 	};
 } // namespace UE::Audio::Insights
diff --git a/Engine/Plugins/Runtime/AudioModulation/Source/AudioModulationEditor/Private/Insights/Providers/ControlBusTraceProvider.cpp b/Engine/Plugins/Runtime/AudioModulation/Source/AudioModulationEditor/Private/Insights/Providers/ControlBusTraceProvider.cpp
--- a/Engine/Plugins/Runtime/AudioModulation/Source/AudioModulationEditor/Private/Insights/Providers/ControlBusTraceProvider.cpp
+++ b/Engine/Plugins/Runtime/AudioModulation/Source/AudioModulationEditor/Private/Insights/Providers/ControlBusTraceProvider.cpp
@@ -1,10 +1,9 @@
 // Copyright Epic Games, Inc. All Rights Reserved.
 #include "Insights/Providers/ControlBusTraceProvider.h"
 
+#include "AudioInsightsUtils.h"
 #include "Trace/Analyzer.h"
 #include "TraceServices/Model/AnalysisSession.h"
-#include "TraceServices/ModuleService.h"
-
 
 namespace AudioModulationEditor
 {
@@ -15,6 +14,8 @@ namespace AudioModulationEditor
 
 	bool FControlBusTraceProvider::ProcessMessages()
 	{
+		using namespace UE::Audio::Insights;
+
 		auto BumpEntryFunc = [this](const FControlBusMessageBase& Msg)
 		{
 			TSharedPtr<FControlBusDashboardEntry>* ToReturn = nullptr;
@@ -38,7 +39,7 @@ namespace AudioModulationEditor
 			FControlBusDashboardEntry& EntryRef = *OutEntry->Get();
 			EntryRef.Name = *Msg.BusName;
 			EntryRef.ControlBusId = Msg.ControlBusId;
-			EntryRef.DisplayName = EntryRef.GetDisplayName().ToString();
+			EntryRef.DisplayName = AudioInsightsUtils::ResolveObjectDisplayName(EntryRef.Name);
 			EntryRef.ParamName = Msg.ParamName;
 		});
 
@@ -47,7 +48,7 @@ namespace AudioModulationEditor
 		{
 			FControlBusDashboardEntry& EntryRef = *OutEntry->Get();
 			EntryRef.Name = *Msg.BusName;
-			EntryRef.DisplayName = EntryRef.GetDisplayName().ToString();
+			EntryRef.DisplayName = AudioInsightsUtils::ResolveObjectDisplayName(EntryRef.Name);
 			EntryRef.ParamName = Msg.ParamName;
 			EntryRef.Value = Msg.Value;
 		});
diff --git a/Engine/Plugins/Runtime/AudioModulation/Source/AudioModulationEditor/Private/Insights/Providers/ModulationMatrixTraceProvider.cpp b/Engine/Plugins/Runtime/AudioModulation/Source/AudioModulationEditor/Private/Insights/Providers/ModulationMatrixTraceProvider.cpp
--- a/Engine/Plugins/Runtime/AudioModulation/Source/AudioModulationEditor/Private/Insights/Providers/ModulationMatrixTraceProvider.cpp
+++ b/Engine/Plugins/Runtime/AudioModulation/Source/AudioModulationEditor/Private/Insights/Providers/ModulationMatrixTraceProvider.cpp
@@ -2,35 +2,12 @@
 #include "Insights/Providers/ModulationMatrixTraceProvider.h"
 
 #include "AudioDeviceManager.h"
+#include "AudioInsightsUtils.h"
 #include "Trace/Analyzer.h"
 #include "TraceServices/Model/AnalysisSession.h"
 
 namespace AudioModulationEditor
 {
-	namespace FModulationMatrixTraceProviderPrivate
-	{
-		FString ResolveDisplayName(const FString& InNamePath)
-		{
-			const FSoftObjectPath SoftObjectPath = FSoftObjectPath(InNamePath);
-
-			if (SoftObjectPath.IsAsset())
-			{
-				return SoftObjectPath.GetAssetName();
-			}
-			else if (SoftObjectPath.IsSubobject())
-			{
-				FString Left;
-				FString Right;
-				if (SoftObjectPath.GetSubPathString().Split(".", &Left, &Right, ESearchCase::IgnoreCase, ESearchDir::FromEnd))
-				{
-					return Right;
-				}
-			}
-
-			return InNamePath;
-		}
-	}
-
 	FModulationMatrixTraceProvider::FModulationMatrixTraceProvider()
 		: UE::Audio::Insights::TDeviceDataMapTraceProvider<uint32, TSharedPtr<FModulationMatrixDashboardEntry>>(GetName_Static())
 	{
@@ -96,7 +73,7 @@ namespace AudioModulationEditor
 
 	bool FModulationMatrixTraceProvider::ProcessMessages()
 	{
-		using namespace FModulationMatrixTraceProviderPrivate;
+		using namespace UE::Audio::Insights;
 
 		// Helper lambdas
 		auto BumpEntryFunc = [this](const FModulationMatrixMessageBase& Msg)
@@ -153,11 +130,11 @@ namespace AudioModulationEditor
 			{
 				if (TMap<FBusId, FString>* FoundBusInfoPtr = ModulatingSourceIdToBusesInfoMap.Find(Msg.ModulatingSourceId))
 				{
-					FoundBusInfoPtr->FindOrAdd(Msg.SourceId, ResolveDisplayName(Msg.BusName));
+					FoundBusInfoPtr->FindOrAdd(Msg.SourceId, AudioInsightsUtils::ResolveObjectDisplayName(Msg.BusName));
 				}
 				else
 				{
-					ModulatingSourceIdToBusesInfoMap.Emplace(Msg.ModulatingSourceId).Emplace(Msg.SourceId, ResolveDisplayName(Msg.BusName));
+					ModulatingSourceIdToBusesInfoMap.Emplace(Msg.ModulatingSourceId).Emplace(Msg.SourceId, AudioInsightsUtils::ResolveObjectDisplayName(Msg.BusName));
 				}
 			});
 
@@ -166,7 +143,7 @@ namespace AudioModulationEditor
 			{
 				FModulationMatrixDashboardEntry& EntryRef = *OutEntry->Get();
 				EntryRef.Name = *Msg.Name;
-				EntryRef.DisplayName = ResolveDisplayName(Msg.Name);
+				EntryRef.DisplayName = AudioInsightsUtils::ResolveObjectDisplayName(Msg.Name);
 				EntryRef.EntryType = EModulationMatrixEntryType::BusMix;
 
 				UpdateActiveModulatorSourceIds(Msg.DeviceId, EntryRef.SourceId);
