Unclamp Volume Modulation

diff --git a/Engine/Plugins/Runtime/AudioModulation/Source/AudioModulation/Private/SoundControlBusMix.cpp b/Engine/Plugins/Runtime/AudioModulation/Source/AudioModulation/Private/SoundControlBusMix.cpp
--- a/Engine/Plugins/Runtime/AudioModulation/Source/AudioModulation/Private/SoundControlBusMix.cpp
+++ b/Engine/Plugins/Runtime/AudioModulation/Source/AudioModulation/Private/SoundControlBusMix.cpp
@@ -28,7 +28,7 @@ FSoundControlBusMixStage::FSoundControlBusMixStage()
 FSoundControlBusMixStage::FSoundControlBusMixStage(USoundControlBus* InBus, const float TargetValue)
 	: Bus(InBus)
 {
-	Value.TargetValue = FMath::Clamp(TargetValue, 0.0f, 1.0f);
+	Value.TargetValue = TargetValue;
 }
 
 USoundControlBusMix::USoundControlBusMix(const FObjectInitializer& ObjectInitializer)
@@ -132,10 +132,20 @@ void USoundControlBusMix::OnPropertyChanged(FProperty* Property, EPropertyChange
 				{
 					if (Stage.Bus)
 					{
-						Stage.Value.TargetValue = FMath::Clamp(Stage.Value.TargetValue, 0.0f, 1.0f);
 						float UnitValue = Stage.Value.TargetValue;
 						if (USoundModulationParameter* Parameter = Stage.Bus->Parameter)
 						{
+							// Clamp the values--either to the normalized [0,1] range or to the parameter-specified "unclamped" range
+							if (!Parameter->Settings.bClampToNormalizedRange)
+							{
+								const TRange<float> Range = Parameter->GetNormalizedLimits();
+								Stage.Value.TargetValue = FMath::Clamp(Stage.Value.TargetValue, Range.GetLowerBoundValue(), Range.GetUpperBoundValue());
+							}
+							else
+							{
+								Stage.Value.TargetValue = FMath::Clamp(Stage.Value.TargetValue, 0.0f, 1.0f);
+							}
+							
 							UnitValue = Parameter->ConvertNormalizedToUnit(Stage.Value.TargetValue);
 						}
 
diff --git a/Engine/Plugins/Runtime/AudioModulation/Source/AudioModulation/Private/SoundControlBusProxy.cpp b/Engine/Plugins/Runtime/AudioModulation/Source/AudioModulation/Private/SoundControlBusProxy.cpp
--- a/Engine/Plugins/Runtime/AudioModulation/Source/AudioModulation/Private/SoundControlBusProxy.cpp
+++ b/Engine/Plugins/Runtime/AudioModulation/Source/AudioModulation/Private/SoundControlBusProxy.cpp
@@ -10,6 +10,14 @@
 #include "Engine/World.h"
 #include "SoundModulationGenerator.h"
 
+static int32 ForceClampNormalizedValuesCVar = 0;
+FAutoConsoleVariableRef CVarForceClampNormalizedValues(
+	TEXT("au.Modulation.ForceClampModulationValues"),
+	ForceClampNormalizedValuesCVar,
+	TEXT("Forces all Modulation values to stay in the normalized [0,1] range.\n")
+	TEXT("0: Normalized values can go outside [0,1] if the Modulation Parameter allows it. 1: Normalized values are always in [0,1] range."),
+	ECVF_Default);
+
 #if UE_AUDIO_PROFILERTRACE_ENABLED
 UE_TRACE_EVENT_BEGIN(Audio, ControlBusActivate)
 	UE_TRACE_EVENT_FIELD(uint32, DeviceId)
@@ -71,14 +79,6 @@ namespace AudioModulation
 		return ModSystem.RegisterModulator(HandleId, *this);
 	}
 
-	FControlBusProxy::FControlBusProxy()
-		: DefaultValue(0.0f)
-		, GeneratorValue(1.0f)
-		, MixValue(NAN)
-		, bBypass(false)
-	{
-	}
-
 	FControlBusProxy::FControlBusProxy(FControlBusSettings&& InSettings, FAudioModulationSystem& InModSystem)
 		: TModulatorProxyRefType(InSettings.GetPath(), InSettings.GetId(), InModSystem)
 	{
@@ -126,8 +126,14 @@ namespace AudioModulation
 
 	float FControlBusProxy::GetValue() const
 	{
-		const float DefaultMixed = Mix(DefaultValue);
-		return FMath::Clamp(DefaultMixed * GeneratorValue, 0.0f, 1.0f);
+		float DefaultMixed = Mix(DefaultValue) * GeneratorValue;
+
+		if (ForceClampNormalizedValuesCVar || bClampToNormalizedRange)
+		{
+			return FMath::Clamp(DefaultMixed, 0.0f, 1.0f);
+		}
+		
+		return DefaultMixed;
 	}
 
 	FName FControlBusProxy::GetParameterName() const
@@ -154,6 +160,8 @@ namespace AudioModulation
 		DefaultValue = FMath::Clamp(InSettings.DefaultValue, 0.0f, 1.0f);
 		bBypass = InSettings.bBypass;
 
+		bClampToNormalizedRange = InSettings.OutputParameter.bClampToNormalizedRange;
+		
 		TArray<FGeneratorHandle> NewHandles;
 		for (FModulationGeneratorSettings& GeneratorSettings : InSettings.GeneratorSettings)
 		{
diff --git a/Engine/Plugins/Runtime/AudioModulation/Source/AudioModulation/Private/SoundControlBusProxy.h b/Engine/Plugins/Runtime/AudioModulation/Source/AudioModulation/Private/SoundControlBusProxy.h
--- a/Engine/Plugins/Runtime/AudioModulation/Source/AudioModulation/Private/SoundControlBusProxy.h
+++ b/Engine/Plugins/Runtime/AudioModulation/Source/AudioModulation/Private/SoundControlBusProxy.h
@@ -91,13 +91,15 @@ namespace AudioModulation
 		void Init(FControlBusSettings&& InSettings);
 		float Mix(float ValueA) const;
 
-		float DefaultValue;
-
+		float DefaultValue = 0.0f;
+		
 		// Cached values
-		float GeneratorValue;
-		float MixValue;
+		float GeneratorValue = 1.0f;
+		float MixValue = NAN;
 
-		bool bBypass;
+		bool bBypass = false;
+
+		bool bClampToNormalizedRange = true;
 
 		Audio::FModulationMixFunction MixFunction;
 		TArray<FGeneratorHandle> GeneratorHandles;
diff --git a/Engine/Plugins/Runtime/AudioModulation/Source/AudioModulation/Private/SoundModulationParameter.cpp b/Engine/Plugins/Runtime/AudioModulation/Source/AudioModulation/Private/SoundModulationParameter.cpp
--- a/Engine/Plugins/Runtime/AudioModulation/Source/AudioModulation/Private/SoundModulationParameter.cpp
+++ b/Engine/Plugins/Runtime/AudioModulation/Source/AudioModulation/Private/SoundModulationParameter.cpp
@@ -63,6 +63,7 @@ Audio::FModulationParameter USoundModulationParameter::CreateParameter() const
 	Parameter.DefaultValue = GetUnitDefault();
 	Parameter.MinValue = GetUnitMin();
 	Parameter.MaxValue = GetUnitMax();
+	Parameter.bClampToNormalizedRange = Settings.bClampToNormalizedRange;
 	Parameter.TypeHash = GetTypeHash(this);
 
 #if WITH_EDITORONLY_DATA
@@ -242,6 +243,14 @@ Audio::FModulationParameter USoundModulationParameterBipolar::CreateDefaultParam
 	return DefaultParam;
 }
 
+USoundModulationParameterVolume::USoundModulationParameterVolume(const FObjectInitializer& ObjectInitializer)
+	: Super(ObjectInitializer)
+{
+#if WITH_EDITORONLY_DATA
+	Settings.bShowClampOption = true;
+#endif
+}
+
 bool USoundModulationParameterVolume::RequiresUnitConversion() const
 {
 	return true;
@@ -249,21 +258,39 @@ bool USoundModulationParameterVolume::RequiresUnitConversion() const
 
 Audio::FModulationUnitConversionFunction USoundModulationParameterVolume::GetUnitConversionFunction() const
 {
-	return [InUnitMin = GetUnitMin()](float& InOutValue)
+	return [InUnitMin = GetUnitMin(), InUnitMax = VolumeLimit, bClampToNormalized = Settings.bClampToNormalizedRange](float& InOutValue)
 	{
 		InOutValue = InOutValue > 0.0f
 			? Audio::ConvertToDecibels(InOutValue)
 			: InUnitMin;
+
+		if (bClampToNormalized)
+		{
+			InOutValue = FMath::Clamp(InOutValue, InUnitMin, 0.0f);
+		}
+		else
+		{
+			InOutValue = FMath::Clamp(InOutValue, InUnitMin, InUnitMax);
+		}
 	};
 }
 
 Audio::FModulationNormalizedConversionFunction USoundModulationParameterVolume::GetNormalizedConversionFunction() const
 {
-	return [InUnitMin = GetUnitMin()](float& InOutValue)
+	return [InUnitMin = GetUnitMin(), InUnitMax = VolumeLimit, bClampToNormalized = Settings.bClampToNormalizedRange](float& InOutValue)
 	{
 		InOutValue = InOutValue < InUnitMin || FMath::IsNearlyEqual(InOutValue, InUnitMin)
 			? 0.0f
 			: Audio::ConvertToLinear(InOutValue);
+
+		if (bClampToNormalized)
+		{
+			InOutValue = FMath::Clamp(InOutValue, 0, 1.0f);
+		}
+		else
+		{
+			InOutValue = FMath::Clamp(InOutValue, 0, Audio::ConvertToLinear(InUnitMax));
+		}
 	};
 }
 
@@ -277,6 +304,11 @@ float USoundModulationParameterVolume::GetUnitMax() const
 	return 0.0f;
 }
 
+TRange<float> USoundModulationParameterVolume::GetNormalizedLimits() const
+{
+	return TRange<float>(0.0f, Audio::ConvertToLinear(VolumeLimit));
+}
+
 Audio::FModulationParameter USoundModulationParameterVolume::CreateDefaultParameter(float MinUnitVolume)
 {
 	UClass* ThisClass = USoundModulationParameterVolume::StaticClass();
@@ -296,6 +328,15 @@ Audio::FModulationParameter USoundModulationParameterVolume::CreateDefaultParame
 	return NewParam;
 }
 
+#if WITH_EDITOR
+void USoundModulationParameterVolume::PostEditChangeProperty(FPropertyChangedEvent& InPropertyChangedEvent)
+{
+	VolumeLimit = FMath::Max(VolumeLimit, MinVolume);
+	
+	Super::PostEditChangeProperty(InPropertyChangedEvent);
+}
+#endif
+
 USoundModulationParameterAdditive::USoundModulationParameterAdditive(const FObjectInitializer& ObjectInitializer)
 	: Super(ObjectInitializer)
 {
diff --git a/Engine/Plugins/Runtime/AudioModulation/Source/AudioModulation/Public/SoundModulationParameter.h b/Engine/Plugins/Runtime/AudioModulation/Source/AudioModulation/Public/SoundModulationParameter.h
--- a/Engine/Plugins/Runtime/AudioModulation/Source/AudioModulation/Public/SoundModulationParameter.h
+++ b/Engine/Plugins/Runtime/AudioModulation/Source/AudioModulation/Public/SoundModulationParameter.h
@@ -22,6 +22,10 @@ struct FSoundModulationParameterSettings
 	UPROPERTY(EditAnywhere, Category = General, BlueprintReadOnly, meta = (ClampMin = "0.0", ClampMax = "1.0"))
 	float ValueNormalized = 1.0f;
 
+	/** Whether or not to clamp values using this parameter to the normalized [0,1] range. */
+	UPROPERTY(EditAnywhere, Category = General, BlueprintReadOnly, meta = (EditCondition = "bShowClampOption"))
+	bool bClampToNormalizedRange = true;
+	
 #if WITH_EDITORONLY_DATA
 	/** (Optional) Text name of parameter's unit */
 	UPROPERTY(EditAnywhere, Category = General)
@@ -35,6 +39,8 @@ struct FSoundModulationParameterSettings
 	  */
 	UPROPERTY(Transient, EditAnywhere, Category = General)
 	float ValueUnit = 1.0f;
+
+	bool bShowClampOption = false;
 #endif // WITH_EDITORONLY_DATA
 };
 
@@ -104,6 +110,14 @@ public:
 		return 1.0f;
 	}
 
+	/** Returns the range to use when FSoundModulationParameterSettings::bClampToNormalizedRange is set to false.
+	 *  If overriden, the Unit/Normalized conversion functions must be written in a way to account for this range.
+	 *  Returned range is unitless/"normalized". */
+	virtual TRange<float> GetNormalizedLimits() const
+	{
+		return TRange<float>(0.0f, 1.0f);
+	}
+
 	//~Begin IAudioProxyDataFactory Interface.
 	UE_API virtual TSharedPtr<Audio::IProxyData> CreateProxyData(const Audio::FProxyDataInitParams& InitParams) override;
 	//~ End IAudioProxyDataFactory Interface.
@@ -243,19 +257,28 @@ public:
 UCLASS(BlueprintType, MinimalAPI)
 class USoundModulationParameterVolume : public USoundModulationParameter
 {
-	GENERATED_BODY()
+	GENERATED_UCLASS_BODY()
 
 public:
-	/** Minimum volume of parameter. Only enforced at modulation destination. */
+	/** Minimum volume of parameter (dB). Only enforced at modulation destination. */
 	UPROPERTY(EditAnywhere, Category = General, BlueprintReadOnly, meta = (ClampMax = 0.0))
 	float MinVolume = -100.0f;
 
+	/** What the maximum volume (in dB) should be if "Clamp to Normalized Range" is set to false. */
+	UPROPERTY(EditAnywhere, Category = General, BlueprintReadOnly)
+	float VolumeLimit = 20.0f;
+	
 	virtual bool RequiresUnitConversion() const override;
 	virtual Audio::FModulationUnitConversionFunction GetUnitConversionFunction() const override;
 	virtual Audio::FModulationNormalizedConversionFunction GetNormalizedConversionFunction() const override;
 	virtual float GetUnitMin() const override;
 	virtual float GetUnitMax() const override;
+	virtual TRange<float> GetNormalizedLimits() const override;
 	static Audio::FModulationParameter CreateDefaultParameter(float MinUnitVolume = -100.0f);
+
+#if WITH_EDITOR
+	virtual void PostEditChangeProperty(FPropertyChangedEvent& InPropertyChangedEvent) override;
+#endif // WITH_EDITOR
 };
 
 // Modulation Parameter whose values are mixed via addition.
diff --git a/Engine/Plugins/Runtime/AudioModulation/Source/AudioModulationEditor/Private/Layouts/SoundModulationParameterSettingsLayout.cpp b/Engine/Plugins/Runtime/AudioModulation/Source/AudioModulationEditor/Private/Layouts/SoundModulationParameterSettingsLayout.cpp
--- a/Engine/Plugins/Runtime/AudioModulation/Source/AudioModulationEditor/Private/Layouts/SoundModulationParameterSettingsLayout.cpp
+++ b/Engine/Plugins/Runtime/AudioModulation/Source/AudioModulationEditor/Private/Layouts/SoundModulationParameterSettingsLayout.cpp
@@ -65,6 +65,7 @@ void FSoundModulationParameterSettingsLayoutCustomization::CustomizeChildren(TSh
 	TSharedRef<IPropertyHandle> NormalizedValueHandle = PropertyHandles.FindChecked(GET_MEMBER_NAME_CHECKED(FSoundModulationParameterSettings, ValueNormalized)).ToSharedRef();
 	TSharedRef<IPropertyHandle> UnitValueHandle = PropertyHandles.FindChecked(GET_MEMBER_NAME_CHECKED(FSoundModulationParameterSettings, ValueUnit)).ToSharedRef();
 	TSharedRef<IPropertyHandle> UnitDisplayNameHandle = PropertyHandles.FindChecked(GET_MEMBER_NAME_CHECKED(FSoundModulationParameterSettings, UnitDisplayName)).ToSharedRef();
+	TSharedRef<IPropertyHandle> ClampToNormalizedHandle = PropertyHandles.FindChecked(GET_MEMBER_NAME_CHECKED(FSoundModulationParameterSettings, bClampToNormalizedRange)).ToSharedRef();
 
 	// Default Value Row (Unit)
 	USoundModulationParameter* Parameter = ModParamSettingsUtils::GetStructParameter(StructPropertyHandle);
@@ -166,8 +167,14 @@ void FSoundModulationParameterSettingsLayoutCustomization::CustomizeChildren(TSh
 			[
 				NormalizedValueHandle->CreatePropertyValueWidget()
 			];
-	}
+	}	
 
+	TAttribute<bool> EditCondition = TAttribute<bool>::Create([Parameter]()
+	{
+		return Parameter->Settings.bShowClampOption;
+	});
+	
 	ChildBuilder.AddProperty(UnitDisplayNameHandle);
+	ChildBuilder.AddProperty(ClampToNormalizedHandle).EditCondition(EditCondition, nullptr);
 }
 #undef LOCTEXT_NAMESPACE // SoundModulationFloat
diff --git a/Engine/Source/Runtime/AudioExtensions/Private/IAudioModulation.cpp b/Engine/Source/Runtime/AudioExtensions/Private/IAudioModulation.cpp
--- a/Engine/Source/Runtime/AudioExtensions/Private/IAudioModulation.cpp
+++ b/Engine/Source/Runtime/AudioExtensions/Private/IAudioModulation.cpp
@@ -87,78 +87,6 @@ namespace Audio
 	{
 	}
 
-	FModulationParameter::FModulationParameter(FModulationParameter&& InParam)
-		: ParameterName(MoveTemp(InParam.ParameterName))
-		, DefaultValue(InParam.DefaultValue)
-		, MinValue(InParam.MinValue)
-		, MaxValue(InParam.MaxValue)
-		, bRequiresConversion(InParam.bRequiresConversion)
-	#if WITH_EDITORONLY_DATA
-		, UnitDisplayName(MoveTemp(InParam.UnitDisplayName))
-		, ClassName(MoveTemp(InParam.ClassName))
-	#endif // WITH_EDITORONLY_DATA
-		, MixFunction(MoveTemp(InParam.MixFunction))
-		, UnitFunction(MoveTemp(InParam.UnitFunction))
-		, NormalizedFunction(MoveTemp(InParam.NormalizedFunction))
-	{
-	}
-
-	FModulationParameter::FModulationParameter(const FModulationParameter& InParam)
-		: ParameterName(InParam.ParameterName)
-		, DefaultValue(InParam.DefaultValue)
-		, MinValue(InParam.MinValue)
-		, MaxValue(InParam.MaxValue)
-		, bRequiresConversion(InParam.bRequiresConversion)
-#if WITH_EDITORONLY_DATA
-		, UnitDisplayName(InParam.UnitDisplayName)
-		, ClassName(InParam.ClassName)
-#endif // WITH_EDITORONLY_DATA
-		, MixFunction(InParam.MixFunction)
-		, UnitFunction(InParam.UnitFunction)
-		, NormalizedFunction(InParam.NormalizedFunction)
-	{
-	}
-
-	FModulationParameter& FModulationParameter::operator=(const FModulationParameter& InParam)
-	{
-		ParameterName = InParam.ParameterName;
-		DefaultValue = InParam.DefaultValue;
-		MinValue = InParam.MinValue;
-		MaxValue = InParam.MaxValue;
-		bRequiresConversion = InParam.bRequiresConversion;
-
-#if WITH_EDITORONLY_DATA
-		UnitDisplayName = InParam.UnitDisplayName;
-		ClassName = InParam.ClassName;
-#endif // WITH_EDITORONLY_DATA
-
-		MixFunction = InParam.MixFunction;
-		UnitFunction = InParam.UnitFunction;
-		NormalizedFunction = InParam.NormalizedFunction;
-
-		return *this;
-	}
-
-	FModulationParameter& FModulationParameter::operator=(FModulationParameter&& InParam)
-	{
-		ParameterName = MoveTemp(InParam.ParameterName);
-		DefaultValue = InParam.DefaultValue;
-		MinValue = InParam.MinValue;
-		MaxValue = InParam.MaxValue;
-		bRequiresConversion = InParam.bRequiresConversion;
-
-	#if WITH_EDITORONLY_DATA
-		UnitDisplayName = MoveTemp(InParam.UnitDisplayName);
-		ClassName = MoveTemp(InParam.ClassName);
-	#endif // WITH_EDITORONLY_DATA
-
-		MixFunction = MoveTemp(InParam.MixFunction);
-		UnitFunction = MoveTemp(InParam.UnitFunction);
-		NormalizedFunction = MoveTemp(InParam.NormalizedFunction);
-
-		return *this;
-	}
-
 	const FModulationMixFunction& FModulationParameter::GetDefaultMixFunction()
 	{
 		static const FModulationMixFunction DefaultMixFunction = [](float& InOutValueA, float InValueB)
diff --git a/Engine/Source/Runtime/AudioExtensions/Public/IAudioModulation.h b/Engine/Source/Runtime/AudioExtensions/Public/IAudioModulation.h
--- a/Engine/Source/Runtime/AudioExtensions/Public/IAudioModulation.h
+++ b/Engine/Source/Runtime/AudioExtensions/Public/IAudioModulation.h
@@ -51,11 +51,11 @@ namespace Audio
 	struct FModulationParameter
 	{
 		AUDIOEXTENSIONS_API FModulationParameter();
-		AUDIOEXTENSIONS_API FModulationParameter(const FModulationParameter& InParam);
-		AUDIOEXTENSIONS_API FModulationParameter(FModulationParameter&& InParam);
+		FModulationParameter(const FModulationParameter& InParam) = default;
+		FModulationParameter(FModulationParameter&& InParam) = default;
 
-		AUDIOEXTENSIONS_API FModulationParameter& operator=(FModulationParameter&& InParam);
-		AUDIOEXTENSIONS_API FModulationParameter& operator=(const FModulationParameter& InParam);
+		FModulationParameter& operator=(FModulationParameter&& InParam) = default;
+		FModulationParameter& operator=(const FModulationParameter& InParam) = default;
 
 		FName ParameterName;
 
@@ -71,6 +71,9 @@ namespace Audio
 		// Whether or not unit conversion is required
 		bool bRequiresConversion = false;
 
+		// Whether or not to clamp normalized values to [0,1] when using this parameter. 
+		bool bClampToNormalizedRange = true;
+		
 		uint32 TypeHash = INDEX_NONE;
 
 #if WITH_EDITORONLY_DATA
