Guard additional clamp calls around Modulation Patch Settings

diff --git a/Engine/Plugins/Runtime/AudioModulation/Source/AudioModulation/Private/SoundModulationParameter.cpp b/Engine/Plugins/Runtime/AudioModulation/Source/AudioModulation/Private/SoundModulationParameter.cpp
--- a/Engine/Plugins/Runtime/AudioModulation/Source/AudioModulation/Private/SoundModulationParameter.cpp
+++ b/Engine/Plugins/Runtime/AudioModulation/Source/AudioModulation/Private/SoundModulationParameter.cpp
@@ -33,21 +33,27 @@ void USoundModulationParameter::PostEditChangeProperty(struct FPropertyChangedEv
 
 void USoundModulationParameter::RefreshNormalizedValue()
 {
-	const float NewNormalizedValue = ConvertUnitToNormalized(Settings.ValueUnit);
-	const float NewNormalizedValueClamped = FMath::Clamp(NewNormalizedValue, 0.0f, 1.0f);
-	if (!FMath::IsNearlyEqual(NewNormalizedValueClamped, Settings.ValueNormalized))
+	float NewNormalizedValue = ConvertUnitToNormalized(Settings.ValueUnit);
+	if (Settings.bClampToNormalizedRange)
 	{
-		Settings.ValueNormalized = NewNormalizedValueClamped;
+		NewNormalizedValue = FMath::Clamp(NewNormalizedValue, 0.0f, 1.0f);
+	}
+	if (!FMath::IsNearlyEqual(NewNormalizedValue, Settings.ValueNormalized))
+	{
+		Settings.ValueNormalized = NewNormalizedValue;
 	}
 }
 
 void USoundModulationParameter::RefreshUnitValue()
 {
-	const float NewUnitValue = ConvertNormalizedToUnit(Settings.ValueNormalized);
-	const float NewUnitValueClamped = FMath::Clamp(NewUnitValue, GetUnitMin(), GetUnitMax());
-	if (!FMath::IsNearlyEqual(NewUnitValueClamped, Settings.ValueUnit))
+	float NewUnitValue = ConvertNormalizedToUnit(Settings.ValueNormalized);
+	if (Settings.bClampToNormalizedRange)
+	{
+		NewUnitValue = FMath::Clamp(NewUnitValue, GetUnitMin(), GetUnitMax());
+	}
+	if (!FMath::IsNearlyEqual(NewUnitValue, Settings.ValueUnit))
 	{
-		Settings.ValueUnit = NewUnitValueClamped;
+		Settings.ValueUnit = NewUnitValue;
 	}
 }
 #endif // WITH_EDITOR
diff --git a/Engine/Plugins/Runtime/AudioModulation/Source/AudioModulationEditor/Private/Layouts/SoundModulationParameterSettingsLayout.cpp b/Engine/Plugins/Runtime/AudioModulation/Source/AudioModulationEditor/Private/Layouts/SoundModulationParameterSettingsLayout.cpp
--- a/Engine/Plugins/Runtime/AudioModulation/Source/AudioModulationEditor/Private/Layouts/SoundModulationParameterSettingsLayout.cpp
+++ b/Engine/Plugins/Runtime/AudioModulation/Source/AudioModulationEditor/Private/Layouts/SoundModulationParameterSettingsLayout.cpp
@@ -171,7 +171,11 @@ void FSoundModulationParameterSettingsLayoutCustomization::CustomizeChildren(TSh
 
 	TAttribute<bool> EditCondition = TAttribute<bool>::Create([Parameter]()
 	{
-		return Parameter->Settings.bShowClampOption;
+		if (Parameter)
+		{
+			return Parameter->Settings.bShowClampOption;
+		}
+		return false;
 	});
 	
 	ChildBuilder.AddProperty(UnitDisplayNameHandle);
