Sequencer: audio control bus track

diff --git a/Engine/Plugins/Runtime/AudioModulation/Source/AudioModulation/AudioModulation.Build.cs b/Engine/Plugins/Runtime/AudioModulation/Source/AudioModulation/AudioModulation.Build.cs
--- a/Engine/Plugins/Runtime/AudioModulation/Source/AudioModulation/AudioModulation.Build.cs
+++ b/Engine/Plugins/Runtime/AudioModulation/Source/AudioModulation/AudioModulation.Build.cs
@@ -17,7 +17,9 @@ namespace UnrealBuildTool.Rules
 				{
 					"AudioExtensions",
 					"Core",
-					"WaveTable"
+					"WaveTable",
+					"MovieSceneTracks",
+					"MovieScene"
 				}
 			);
 
diff --git a/Engine/Plugins/Runtime/AudioModulation/Source/AudioModulation/Private/AudioModulation.cpp b/Engine/Plugins/Runtime/AudioModulation/Source/AudioModulation/Private/AudioModulation.cpp
--- a/Engine/Plugins/Runtime/AudioModulation/Source/AudioModulation/Private/AudioModulation.cpp
+++ b/Engine/Plugins/Runtime/AudioModulation/Source/AudioModulation/Private/AudioModulation.cpp
@@ -15,6 +15,7 @@
 #include "SoundModulationPatch.h"
 #include "SoundModulatorAsset.h"
 #include "UObject/NoExportTypes.h"
+#include "AudioModulationStatics.h"
 
 #if WITH_AUDIOMODULATION_METASOUND_SUPPORT
 #include "MetasoundDataTypeRegistrationMacro.h"
@@ -25,6 +26,8 @@ REGISTER_METASOUND_DATATYPE(AudioModulation::FSoundModulatorAsset, "Modulator",
 REGISTER_METASOUND_DATATYPE(AudioModulation::FSoundModulationParameterAsset, "ModulationParameter", Metasound::ELiteralType::UObjectProxy, USoundModulationParameter);
 #endif // WITH_AUDIOMODULATION_METASOUND_SUPPORT
 
+
+
 namespace AudioModulationCVars
 {
 	static bool bAudioModulationEnabledCVar = true;
@@ -40,6 +43,7 @@ namespace AudioModulation
 	FAudioModulationManager::FAudioModulationManager()
 		: ModSystem(new FAudioModulationSystem())
 	{
+
 	}
 
 	FAudioModulationManager::~FAudioModulationManager()
@@ -57,6 +61,75 @@ namespace AudioModulation
 		ModSystem->OnAuditionEnd();
 	}
 
+	void FAudioModulationManager::GetControlBusAssetPath(TArray<FTopLevelAssetPath>& OutClassNames)
+	{
+		OutClassNames.Add(USoundControlBus::StaticClass()->GetClassPathName());
+	}
+
+	void FAudioModulationManager::CreateNewSequencerBusMix(const UObject* WorldContextObject, FName Name, const TArray<USoundModulatorBase*>& Buses, float Value, float AttackTime, float ReleaseTime, bool bActivate)
+	{
+		UWorld* World = UAudioModulationStatics::GetAudioWorld(WorldContextObject);
+		TArray<USoundControlBus*> BusArray;
+
+		for (USoundModulatorBase* Bus : Buses)
+		{
+			USoundControlBus* CastBus = Cast<USoundControlBus>(Bus);
+
+			if (CastBus != nullptr)
+			{
+				BusArray.Add(CastBus);
+			}
+		}
+
+		if (USoundControlBusMix* NewMix = CreateBusMixFromValue(Name, BusArray, Value, AttackTime, ReleaseTime))
+		{
+			// If we already have an entry with this ID for some reason, clear it to avoid untracked handles before adding the new entry
+			if (SequencerBusMixes.Contains(Name))
+			{
+				ClearSequencerBusMix(Name);
+			}
+
+			if (bActivate)
+			{
+				ModSystem->ActivateBusMix(*NewMix);
+			}
+
+			SequencerBusMixes.Add(Name, NewMix);
+		}
+	}
+
+	void FAudioModulationManager::UpdateMixByModulator(const UObject* WorldContextObject, FName Mix, USoundModulatorBase& Modulator, float Value, float FadeTime)
+	{
+		if (SequencerBusMixes.Contains(Mix))
+		{
+			USoundControlBusMix* ControlBusMix = SequencerBusMixes[Mix];
+			USoundControlBus* CastBus = Cast<USoundControlBus>(&Modulator);
+
+			check(ControlBusMix);
+			check(CastBus);
+
+			if (CastBus == nullptr || ControlBusMix == nullptr)
+			{
+				return;
+			}
+
+			UpdateMixByFilter(CastBus->Address, nullptr, CastBus->Parameter, Value, FadeTime, *ControlBusMix);
+		}
+	}
+
+	void FAudioModulationManager::ClearSequencerBusMix(FName Name)
+	{
+		if (SequencerBusMixes.Contains(Name))
+		{
+			USoundControlBusMix* Mix = SequencerBusMixes[Name];
+			check(Mix);
+
+			ModSystem->DeactivateBusMix(*Mix);
+
+			SequencerBusMixes.Remove(Name);
+		}
+	}
+
 	void FAudioModulationManager::ActivateBus(const USoundControlBus& InBus)
 	{
 		UE_LOG(LogAudioModulation, VeryVerbose, TEXT("FAudioModulationManager::ActivateBus is deprecated. Use UAudioModulationDestination API (see SetModulator) instead to enforce bus lifetime"));
diff --git a/Engine/Plugins/Runtime/AudioModulation/Source/AudioModulation/Private/MovieSceneAudioControlBusSection.cpp b/Engine/Plugins/Runtime/AudioModulation/Source/AudioModulation/Private/MovieSceneAudioControlBusSection.cpp
new file mode 100644
--- /dev/null
+++ b/Engine/Plugins/Runtime/AudioModulation/Source/AudioModulation/Private/MovieSceneAudioControlBusSection.cpp
@@ -0,0 +1,66 @@
+// Copyright Epic Games, Inc. All Rights Reserved.
+
+#include "MovieSceneAudioControlBusSection.h"
+
+#include "Channels/MovieSceneChannelProxy.h"
+#include "EntitySystem/BuiltInComponentTypes.h"
+#include "MovieScene.h"
+#include "MovieSceneTracksComponentTypes.h"
+#include "MovieSceneAudioControlBusTrack.h"
+#include "MovieSceneAudioModulationComponentTypes.h"
+
+#include UE_INLINE_GENERATED_CPP_BY_NAME(MovieSceneAudioControlBusSection)
+
+UMovieSceneAudioControlBusSection::UMovieSceneAudioControlBusSection(const FObjectInitializer& ObjectInitializer)
+	: Super(ObjectInitializer)
+{
+	BusValue.SetDefault(1.f);
+	MixBus = nullptr;
+
+	EvalOptions.EnableAndSetCompletionMode(EMovieSceneCompletionMode::ProjectDefault);
+}
+
+EMovieSceneChannelProxyType UMovieSceneAudioControlBusSection::CacheChannelProxy()
+{
+	// Set up the channel proxy
+	FMovieSceneChannelProxyData Channels;
+
+#if WITH_EDITOR
+	FMovieSceneChannelMetaData Data;
+	Data.SetIdentifiers(TEXT("Control Value"), FText::FromName(MixBus == nullptr ? TEXT("Control Bus") : MixBus->GetFName()));
+
+	Channels.Add(BusValue, Data, TMovieSceneExternalValue<float>::Make());
+#else
+	Channels.Add(BusValue);
+
+#endif
+
+	ChannelProxy = MakeShared<FMovieSceneChannelProxy>(MoveTemp(Channels));
+
+	return EMovieSceneChannelProxyType::Dynamic;
+}
+
+void UMovieSceneAudioControlBusSection::ImportEntityImpl(UMovieSceneEntitySystemLinker* EntityLinker, const FEntityImportParams& Params, FImportedEntity* OutImportedEntity)
+{
+	using namespace UE::MovieScene;
+	const FBuiltInComponentTypes* BuiltInComponents = FBuiltInComponentTypes::Get();
+	const FMovieSceneAudioModulationComponentTypes* TrackComponents = FMovieSceneAudioModulationComponentTypes::Get();
+
+	check(BuiltInComponents);
+	check(TrackComponents);
+
+	const FGuid ObjectBindingID = Params.GetObjectBindingID();
+
+	OutImportedEntity->AddBuilder(
+		FEntityBuilder()
+		.AddConditional(BuiltInComponents->GenericObjectBinding, ObjectBindingID, ObjectBindingID.IsValid())
+		.AddTagConditional(BuiltInComponents->Tags.Root, !ObjectBindingID.IsValid())
+		.Add(TrackComponents->AudioControlBus, FMovieSceneAudioControlBusComponentData{ this })
+		.Add(BuiltInComponents->FloatChannel[0], &BusValue) // Float channels have 9 slots to store and send data, so since we have 1 float channel we send data through the initial float channel
+	);
+}
+
+void UMovieSceneAudioControlBusSection::SetBus(USoundModulatorBase* NewBus)
+{
+	MixBus = NewBus;
+}
diff --git a/Engine/Plugins/Runtime/AudioModulation/Source/AudioModulation/Private/MovieSceneAudioControlBusSystem.cpp b/Engine/Plugins/Runtime/AudioModulation/Source/AudioModulation/Private/MovieSceneAudioControlBusSystem.cpp
new file mode 100644
--- /dev/null
+++ b/Engine/Plugins/Runtime/AudioModulation/Source/AudioModulation/Private/MovieSceneAudioControlBusSystem.cpp
@@ -0,0 +1,283 @@
+// Copyright Epic Games, Inc. All Rights Reserved.
+
+#include "MovieSceneAudioControlBusSystem.h"
+
+#include "AudioDevice.h"
+#include "AudioDeviceManager.h"
+#include "Engine/EngineTypes.h"
+#include "EntitySystem/BuiltInComponentTypes.h"
+#include "EntitySystem/IMovieSceneTaskScheduler.h"
+#include "EntitySystem/MovieSceneEntitySystemLinker.h"
+#include "EntitySystem/MovieSceneEntitySystemTask.h"
+#include "EntitySystem/MovieSceneInstanceRegistry.h"
+#include "Evaluation/PreAnimatedState/MovieScenePreAnimatedObjectStorage.h"
+#include "Evaluation/PreAnimatedState/MovieScenePreAnimatedStateStorage.h"
+#include "Evaluation/PreAnimatedState/MovieScenePreAnimatedStorageID.inl"
+#include "IAudioModulation.h"
+#include "EntitySystem/MovieSceneEntitySystemRunner.h"
+#include "MovieSceneTracksComponentTypes.h"
+#include "MovieSceneAudioModulationComponentTypes.h"
+
+#include UE_INLINE_GENERATED_CPP_BY_NAME(MovieSceneAudioControlBusSystem)
+
+using AudioControlBusSectionData = UE::MovieScene::AudioControlBusSectionData;
+
+namespace UE::MovieScene
+{
+	enum class EPreAnimatedAudioControlBusStateType
+	{
+		ControlBusLifespan
+	};
+
+	template<typename BaseTraits>
+	struct FPreAnimatedAudioControlBusStateTraits : BaseTraits
+	{
+		using KeyType = FObjectKey;
+		using StorageType = EPreAnimatedAudioControlBusStateType;
+
+		EPreAnimatedAudioControlBusStateType CachePreAnimatedValue(FObjectKey InKey)
+		{
+			return EPreAnimatedAudioControlBusStateType::ControlBusLifespan;
+		}
+
+		void RestorePreAnimatedValue(FObjectKey InKey, EPreAnimatedAudioControlBusStateType InStateType, const FRestoreStateParams& Params)
+		{
+		}
+	};
+
+	using FPreAnimatedBoundObjectAudioControlBusStateTraits = FPreAnimatedAudioControlBusStateTraits<FBoundObjectPreAnimatedStateTraits>;
+
+	struct FPreAnimatedAudioControlBusStorage : TPreAnimatedStateStorage_ObjectTraits<FPreAnimatedBoundObjectAudioControlBusStateTraits>
+	{
+		static TAutoRegisterPreAnimatedStorageID<FPreAnimatedAudioControlBusStorage> StorageID;
+	};
+	TAutoRegisterPreAnimatedStorageID<FPreAnimatedAudioControlBusStorage> FPreAnimatedAudioControlBusStorage::StorageID;
+}
+
+namespace UE::MovieScene
+{
+	struct FEvaluateAudioControlBus
+	{
+		UMovieSceneAudioControlBusSystem* AudioSystem;
+		const FInstanceRegistry* InstanceRegistry;
+
+		FEvaluateAudioControlBus(UMovieSceneAudioControlBusSystem* InAudioSystem)
+			: AudioSystem(InAudioSystem)
+		{
+			check(AudioSystem != nullptr);
+			InstanceRegistry = AudioSystem->GetLinker()->GetInstanceRegistry();
+		}
+
+		void ForEachAllocation(
+			const FEntityAllocation* Allocation,
+			TRead<FMovieSceneEntityID> EntityIDs,
+			TRead<FRootInstanceHandle> RootInstanceHandles,
+			TRead<FInstanceHandle> InstanceHandles,
+			TRead<FMovieSceneAudioControlBusComponentData> AudioDatas,
+			TRead<double> BusValue) const
+		{
+			check(Allocation != nullptr);
+			check(EntityIDs != nullptr);
+			check(RootInstanceHandles != nullptr);
+			check(InstanceHandles != nullptr);
+			check(AudioDatas != nullptr);
+
+			const int32 Num = Allocation->Num();
+
+			for (int32 Index = 0; Index < Num; ++Index)
+			{
+				const FMovieSceneEntityID& EntityID = EntityIDs[Index];
+				const FRootInstanceHandle& RootInstanceHandle = RootInstanceHandles[Index];
+				const FInstanceHandle& InstanceHandle = InstanceHandles[Index];
+
+				const FSequenceInstance& Instance = InstanceRegistry->GetInstance(InstanceHandle);
+
+				double BusValueData = BusValue[Index];
+
+				Evaluate(EntityID, AudioDatas[Index], Instance, RootInstanceHandle, BusValueData);
+			}
+		}
+
+	private:
+
+		void Evaluate(
+			const FMovieSceneEntityID& EntityID,
+			const FMovieSceneAudioControlBusComponentData& AudioData,
+			const FSequenceInstance& Instance,
+			const FRootInstanceHandle& RootInstanceHandle,
+			double BusValue) const
+		{
+			UMovieSceneAudioControlBusSection* AudioSection = AudioData.Section;
+			if (!ensureMsgf(AudioSection, TEXT("No valid control bus section found in control bus track component data!")))
+			{
+				return;
+			}
+
+			FInstanceHandle InstanceHandle(Instance.GetInstanceHandle());
+			FObjectKey SectionKey(AudioSection);
+
+			// Grab the cached modulator
+			const AudioControlBusSectionData* EvaluationData = AudioSystem->GetAudioControlBusComponentEvaluationData(SectionKey);
+
+			check(AudioData.Section->MixBus);
+
+			FString MixBusID = FString::Printf(TEXT("%s_%d"), *AudioData.Section->MixBus.GetFName().ToString(), Instance.GetSequenceID().GetInternalValue());
+
+			const UObject* PlaybackContext = Instance.GetSharedPlaybackState()->GetPlaybackContext();
+			check(PlaybackContext);
+			const UWorld* World = PlaybackContext ? PlaybackContext->GetWorld() : nullptr;
+			check(World);
+
+			// If no cached value exists, we need to initialize and put it in pre animated storage so it gets cleared correctly when sequence ends
+			if (!EvaluationData)
+			{
+				EvaluationData = AudioSystem->AddAudioControlBusComponentEvaluationData(InstanceHandle, AudioSection, AudioData.Section->MixBus, FName(*MixBusID));
+
+				AudioSystem->PreAnimatedStorage->BeginTrackingEntity(EntityID, true, RootInstanceHandle, AudioData.Section->MixBus);
+				AudioSystem->PreAnimatedStorage->CachePreAnimatedValue(
+					FCachePreAnimatedValueParams(), AudioData.Section->MixBus,
+					[](FObjectKey InKey) { return EPreAnimatedAudioControlBusStateType::ControlBusLifespan; });
+
+				check(AudioSystem->ModulationManager);
+				AudioSystem->ModulationManager->CreateNewSequencerBusMix(World, FName(*MixBusID), {AudioData.Section->MixBus}, BusValue);
+			}
+
+			check(EvaluationData);
+
+			const FAudioDeviceManager* DeviceManager = FAudioDeviceManager::Get();
+			if (ensure(DeviceManager))
+			{
+				const FAudioDevice* AudioDevice = World->GetAudioDevice().GetAudioDevice();
+				if (ensure(AudioDevice))
+				{
+					if (AudioDevice->IsModulationPluginEnabled())
+					{
+						if (IAudioModulationManager* Modulation = AudioDevice->ModulationInterface.Get())
+						{
+							Modulation->UpdateMixByModulator(World, FName(*MixBusID), *AudioData.Section->MixBus, BusValue);
+						}
+					}
+				}
+			}
+		}
+	};
+}
+
+UMovieSceneAudioControlBusSystem::UMovieSceneAudioControlBusSystem(const FObjectInitializer& ObjInit)
+	: UMovieSceneEntitySystem(ObjInit)
+{
+	using namespace UE::MovieScene;
+
+	RelevantComponent = FMovieSceneAudioModulationComponentTypes::Get()->AudioControlBus;
+	Phase = ESystemPhase::Scheduling;
+
+	if (HasAnyFlags(RF_ClassDefaultObject))
+	{
+		const FBuiltInComponentTypes* BuiltInComponents = FBuiltInComponentTypes::Get();
+		constexpr int FloatChannelCount = 9;
+
+		for (int32 Index = 0; Index < FloatChannelCount; ++Index)
+		{
+			DefineComponentConsumer(GetClass(), BuiltInComponents->DoubleResult[Index]);
+		}
+	}
+}
+
+void UMovieSceneAudioControlBusSystem::OnSchedulePersistentTasks(UE::MovieScene::IEntitySystemScheduler* TaskScheduler)
+{
+	using namespace UE::MovieScene;
+
+	const FBuiltInComponentTypes* BuiltInComponents = FBuiltInComponentTypes::Get();
+	const FMovieSceneAudioModulationComponentTypes* TrackComponents = FMovieSceneAudioModulationComponentTypes::Get();
+
+	// Evaluate the control bus during play using the first double result value from the float channel
+	FTaskID EvaluateAudioTask = FEntityTaskBuilder()
+		.ReadEntityIDs()
+		.Read(BuiltInComponents->RootInstanceHandle)
+		.Read(BuiltInComponents->InstanceHandle)
+		.Read(TrackComponents->AudioControlBus)
+		.Read(BuiltInComponents->DoubleResult[0]) // Bus Value
+		.SetDesiredThread(Linker->EntityManager.GetGatherThread())
+		.Schedule_PerAllocation<FEvaluateAudioControlBus>(&Linker->EntityManager, TaskScheduler, this);
+}
+
+void UMovieSceneAudioControlBusSystem::OnRun(FSystemTaskPrerequisites& InPrerequisites, FSystemSubsequentTasks& Subsequents)
+{
+	using namespace UE::MovieScene;
+
+	const FBuiltInComponentTypes* BuiltInComponents = FBuiltInComponentTypes::Get();
+	const FMovieSceneAudioModulationComponentTypes* TrackComponents = FMovieSceneAudioModulationComponentTypes::Get();
+
+	FSystemTaskPrerequisites Prereqs;
+
+	// Evaluate the control bus during play using the first double result value from the float channel
+	FEntityTaskBuilder()
+		.ReadEntityIDs()
+		.Read(BuiltInComponents->RootInstanceHandle)
+		.Read(BuiltInComponents->InstanceHandle)
+		.Read(TrackComponents->AudioControlBus)
+		.Read(BuiltInComponents->DoubleResult[0]) // Bus Value
+		.SetDesiredThread(Linker->EntityManager.GetGatherThread())
+		.template Dispatch_PerAllocation<FEvaluateAudioControlBus>(&Linker->EntityManager, Prereqs, &Subsequents, this);
+}
+
+void UMovieSceneAudioControlBusSystem::OnLink()
+{
+	using namespace UE::MovieScene;
+
+	PreAnimatedStorage = Linker->PreAnimatedState.GetOrCreateStorage<FPreAnimatedAudioControlBusStorage>();
+
+	// Get the modulation manager and create a mix for our sequence
+	FAudioDeviceManager* DeviceManager = FAudioDeviceManager::Get();
+	if (ensure(DeviceManager))
+	{
+		const FAudioDevice* AudioDevice = DeviceManager->GetActiveAudioDevice().GetAudioDevice();
+		if (ensure(AudioDevice))
+		{
+			if (AudioDevice->IsModulationPluginEnabled())
+			{
+				if (IAudioModulationManager* Modulation = AudioDevice->ModulationInterface.Get())
+				{
+					ModulationManager = Modulation;
+					check(ModulationManager);
+				}
+			}
+		}
+	}
+}
+
+void UMovieSceneAudioControlBusSystem::OnUnlink()
+{
+	ClearBusMixes();
+}
+
+void UMovieSceneAudioControlBusSystem::ClearBusMixes()
+{
+	if (ModulationManager)
+	{
+		for (const TPair<FObjectKey, AudioControlBusSectionData>& Modulator : ModulatorsByObjectKey)
+		{
+			ModulationManager->ClearSequencerBusMix(Modulator.Value.ControlBusMixName);
+		}
+	}
+
+	ModulatorsByObjectKey.Reset();
+	ModulationManager = nullptr;
+}
+
+const AudioControlBusSectionData* UMovieSceneAudioControlBusSystem::GetAudioControlBusComponentEvaluationData(const FObjectKey SectionKey) const
+{
+	const AudioControlBusSectionData* data = ModulatorsByObjectKey.Find(SectionKey);
+
+	return data;
+}
+
+AudioControlBusSectionData* UMovieSceneAudioControlBusSystem::AddAudioControlBusComponentEvaluationData(FInstanceHandle InstanceHandle, FObjectKey SectionKey, USoundModulatorBase* SectionData, FName ID)
+{
+	using namespace UE::MovieScene;
+	check(SectionData != nullptr);
+
+	AudioControlBusSectionData NewData = AudioControlBusSectionData(SectionData, ID);
+
+	return &ModulatorsByObjectKey.Add(SectionKey, NewData);
+}
diff --git a/Engine/Plugins/Runtime/AudioModulation/Source/AudioModulation/Private/MovieSceneAudioControlBusSystem.h b/Engine/Plugins/Runtime/AudioModulation/Source/AudioModulation/Private/MovieSceneAudioControlBusSystem.h
new file mode 100644
--- /dev/null
+++ b/Engine/Plugins/Runtime/AudioModulation/Source/AudioModulation/Private/MovieSceneAudioControlBusSystem.h
@@ -0,0 +1,54 @@
+// Copyright Epic Games, Inc. All Rights Reserved.
+
+#pragma once
+#include "CoreTypes.h"
+#include "EntitySystem/MovieSceneEntityIDs.h"
+#include "EntitySystem/MovieSceneEntitySystem.h"
+#include "EntitySystem/MovieSceneSequenceInstanceHandle.h"
+#include "MovieSceneTracksComponentTypes.h"
+#include "UObject/ObjectMacros.h"
+#include "IAudioModulation.h"
+
+#include "MovieSceneAudioControlBusSystem.generated.h"
+
+namespace UE::MovieScene
+{
+	struct FPreAnimatedAudioControlBusStorage;
+
+	struct AudioControlBusSectionData
+	{
+		TObjectPtr<USoundModulatorBase> ControlBus;
+		FName ControlBusMixName;
+	};
+}
+
+UCLASS()
+class UMovieSceneAudioControlBusSystem : public UMovieSceneEntitySystem
+{
+	GENERATED_BODY()
+
+public:
+	using FInstanceHandle = UE::MovieScene::FInstanceHandle;
+	UMovieSceneAudioControlBusSystem(const FObjectInitializer& ObjInit);
+
+	virtual void OnSchedulePersistentTasks(UE::MovieScene::IEntitySystemScheduler* TaskScheduler) override;
+	virtual void OnRun(FSystemTaskPrerequisites& InPrerequisites, FSystemSubsequentTasks& Subsequents) override;
+
+	virtual void OnLink() override;
+	virtual void OnUnlink() override;
+
+	void ClearBusMixes();
+
+	using AudioControlBusSectionData = UE::MovieScene::AudioControlBusSectionData;
+
+	const AudioControlBusSectionData* GetAudioControlBusComponentEvaluationData(const FObjectKey SectionKey) const;
+	AudioControlBusSectionData* AddAudioControlBusComponentEvaluationData(FInstanceHandle InstanceHandle, FObjectKey SectionKey, USoundModulatorBase* SectionData, FName ID);
+
+	IAudioModulationManager* ModulationManager = nullptr;
+	TArray<USoundModulatorBase*> Modulators;
+
+	/** Pre-animated state */
+	TSharedPtr<UE::MovieScene::FPreAnimatedAudioControlBusStorage> PreAnimatedStorage;
+
+	TMap<FObjectKey, AudioControlBusSectionData> ModulatorsByObjectKey;
+};
\ No newline at end of file
diff --git a/Engine/Plugins/Runtime/AudioModulation/Source/AudioModulation/Private/MovieSceneAudioControlBusTrack.cpp b/Engine/Plugins/Runtime/AudioModulation/Source/AudioModulation/Private/MovieSceneAudioControlBusTrack.cpp
new file mode 100644
--- /dev/null
+++ b/Engine/Plugins/Runtime/AudioModulation/Source/AudioModulation/Private/MovieSceneAudioControlBusTrack.cpp
@@ -0,0 +1,87 @@
+// Copyright Epic Games, Inc. All Rights Reserved.
+
+#include "MovieSceneAudioControlBusTrack.h"
+
+#include "Decorations/MovieSceneSectionAnchorsDecoration.h"
+#include "MovieSceneAudioControlBusSection.h"
+
+#include UE_INLINE_GENERATED_CPP_BY_NAME(MovieSceneAudioControlBusTrack)
+
+UMovieSceneAudioControlBusTrack::UMovieSceneAudioControlBusTrack(const FObjectInitializer& ObjectInitializer)
+	: Super(ObjectInitializer)
+{
+}
+
+UMovieSceneSection* UMovieSceneAudioControlBusTrack::AddNewControlBus(USoundModulatorBase* Bus, FFrameNumber Time, int32 RowIndex)
+{
+	check(Bus);
+
+	// add the section
+	UMovieSceneAudioControlBusSection* NewSection = Cast<UMovieSceneAudioControlBusSection>(CreateNewSection());
+
+	NewSection->SetBus(Bus);
+	constexpr float DefaultDuration = 5.0f;
+	NewSection->InitialPlacementOnRow(AudioSections, Time, DefaultDuration, RowIndex);
+	AudioSections.Add(NewSection);
+
+#if WITH_EDITORONLY_DATA
+	SetTrackRowDisplayName(FText::FromName(Bus->GetFName()), NewSection->GetRowIndex());
+#endif
+
+	return NewSection;
+}
+
+bool UMovieSceneAudioControlBusTrack::SupportsType(TSubclassOf<UMovieSceneSection> SectionClass) const
+{
+	return SectionClass == UMovieSceneAudioControlBusSection::StaticClass();
+}
+
+void UMovieSceneAudioControlBusTrack::RemoveAllAnimationData()
+{
+	AudioSections.Empty();
+}
+
+bool UMovieSceneAudioControlBusTrack::HasSection(const UMovieSceneSection& Section) const
+{
+	return AudioSections.Contains(&Section);
+}
+
+void UMovieSceneAudioControlBusTrack::AddSection(UMovieSceneSection& Section)
+{
+	AudioSections.Add(&Section);
+}
+
+void UMovieSceneAudioControlBusTrack::RemoveSection(UMovieSceneSection& Section)
+{
+	AudioSections.Remove(&Section);
+}
+
+void UMovieSceneAudioControlBusTrack::RemoveSectionAt(int32 SectionIndex)
+{
+	AudioSections.RemoveAt(SectionIndex);
+}
+
+bool UMovieSceneAudioControlBusTrack::IsEmpty() const
+{
+	return AudioSections.IsEmpty();
+}
+
+const TArray<UMovieSceneSection*>& UMovieSceneAudioControlBusTrack::GetAllSections() const
+{
+	return AudioSections;
+}
+
+bool UMovieSceneAudioControlBusTrack::SupportsMultipleRows() const
+{
+	return true;
+}
+
+UMovieSceneSection* UMovieSceneAudioControlBusTrack::CreateNewSection()
+{
+	return NewObject<UMovieSceneAudioControlBusSection>(this, NAME_None, RF_Transactional);
+}
+
+void UMovieSceneAudioControlBusTrack::GetCompatibleUserDecorationsImpl(TSet<UClass*>& OutClasses) const
+{
+	OutClasses.Add(UMovieSceneSectionAnchorsDecoration::StaticClass());
+}
diff --git a/Engine/Plugins/Runtime/AudioModulation/Source/AudioModulation/Private/MovieSceneAudioModulationComponentTypes.cpp b/Engine/Plugins/Runtime/AudioModulation/Source/AudioModulation/Private/MovieSceneAudioModulationComponentTypes.cpp
new file mode 100644
--- /dev/null
+++ b/Engine/Plugins/Runtime/AudioModulation/Source/AudioModulation/Private/MovieSceneAudioModulationComponentTypes.cpp
@@ -0,0 +1,40 @@
+// Copyright Epic Games, Inc. All Rights Reserved.
+
+#include "MovieSceneAudioModulationComponentTypes.h"
+#include "EntitySystem/BuiltInComponentTypes.h"
+#include "EntitySystem/MovieSceneComponentRegistry.h"
+#include "EntitySystem/MovieSceneEntityFactoryTemplates.h"
+#include "EntitySystem/MovieSceneEntitySystemLinker.h"
+#include "EntitySystem/MovieScenePropertyComponentHandler.h"
+
+static bool GMovieSceneAudioModulationComponentTypesDestroyed = false;
+static TUniquePtr<FMovieSceneAudioModulationComponentTypes> GMovieSceneAudioModulationComponentTypes;
+
+FMovieSceneAudioModulationComponentTypes::FMovieSceneAudioModulationComponentTypes()
+{
+	using namespace UE::MovieScene;
+
+	FComponentRegistry* ComponentRegistry = UMovieSceneEntitySystemLinker::GetComponents();
+
+	ComponentRegistry->NewComponentType(&AudioControlBus, TEXT("Audio Control Bus"), EComponentTypeFlags::CopyToChildren);
+}
+
+FMovieSceneAudioModulationComponentTypes::~FMovieSceneAudioModulationComponentTypes()
+{
+}
+
+FMovieSceneAudioModulationComponentTypes* FMovieSceneAudioModulationComponentTypes::Get()
+{
+	if (!GMovieSceneAudioModulationComponentTypes.IsValid())
+	{
+		check(!GMovieSceneAudioModulationComponentTypesDestroyed);
+		GMovieSceneAudioModulationComponentTypes.Reset(new FMovieSceneAudioModulationComponentTypes);
+	}
+	return GMovieSceneAudioModulationComponentTypes.Get();
+}
+
+void FMovieSceneAudioModulationComponentTypes::Destroy()
+{
+	GMovieSceneAudioModulationComponentTypes.Reset();
+	GMovieSceneAudioModulationComponentTypesDestroyed = true;
+}
diff --git a/Engine/Plugins/Runtime/AudioModulation/Source/AudioModulation/Public/AudioModulation.h b/Engine/Plugins/Runtime/AudioModulation/Source/AudioModulation/Public/AudioModulation.h
--- a/Engine/Plugins/Runtime/AudioModulation/Source/AudioModulation/Public/AudioModulation.h
+++ b/Engine/Plugins/Runtime/AudioModulation/Source/AudioModulation/Public/AudioModulation.h
@@ -6,6 +6,7 @@
 #include "HAL/ThreadSafeBool.h"
 #include "IAudioExtensionPlugin.h"
 #include "IAudioModulation.h"
+#include "AudioModulationDestination.h"
 #include "Modules/ModuleInterface.h"
 #include "SoundControlBus.h"
 #include "SoundControlBusMix.h"
@@ -37,6 +38,26 @@ namespace AudioModulation
 		UE_API virtual void Initialize(const FAudioPluginInitializationParams& InitializationParams) override;
 		UE_API virtual void OnAuditionEnd() override;
 
+		UE_API virtual void GetControlBusAssetPath(TArray<FTopLevelAssetPath>& OutClassNames) override;
+
+		UE_API virtual void CreateNewSequencerBusMix(
+			const UObject* WorldContextObject,
+			FName Name,
+			const TArray<USoundModulatorBase*>& Buses,
+			float Value = 1.0f,
+			float AttackTime = 0.1f,
+			float ReleaseTime = 0.1f,
+			bool bActivate = true) override;
+
+		UE_API virtual void UpdateMixByModulator(
+			const UObject* WorldContextObject,
+			FName Mix,
+			USoundModulatorBase& Modulator,
+			float Value = 1.0f,
+			float FadeTime = -1.0f);
+
+		UE_API virtual void ClearSequencerBusMix(FName Name) override;
+
 #if !UE_BUILD_SHIPPING
 		UE_API virtual void SetDebugBusFilter(const FString* InFilter);
 		UE_API virtual void SetDebugMixFilter(const FString* InFilter);
@@ -103,6 +124,7 @@ namespace AudioModulation
 
 	private:
 		FAudioModulationSystem* ModSystem = nullptr;
+		TMap<FName, TObjectPtr<USoundControlBusMix>> SequencerBusMixes;
 	};
 
 	AUDIOMODULATION_API FAudioModulationManager* GetDeviceModulationManager(Audio::FDeviceId InDeviceId);
diff --git a/Engine/Plugins/Runtime/AudioModulation/Source/AudioModulation/Public/MovieSceneAudioControlBusSection.h b/Engine/Plugins/Runtime/AudioModulation/Source/AudioModulation/Public/MovieSceneAudioControlBusSection.h
new file mode 100644
--- /dev/null
+++ b/Engine/Plugins/Runtime/AudioModulation/Source/AudioModulation/Public/MovieSceneAudioControlBusSection.h
@@ -0,0 +1,34 @@
+// Copyright Epic Games, Inc. All Rights Reserved.
+
+#pragma once
+
+#include "Channels/MovieSceneFloatChannel.h"
+#include "EntitySystem/IMovieSceneEntityProvider.h"
+#include "IAudioModulation.h"
+#include "MovieSceneSection.h"
+#include "UObject/ObjectMacros.h"
+
+#include "MovieSceneAudioControlBusSection.generated.h"
+
+UCLASS(MinimalAPI)
+class UMovieSceneAudioControlBusSection 
+	: public UMovieSceneSection
+	, public IMovieSceneEntityProvider
+{
+	GENERATED_UCLASS_BODY()
+
+public:
+
+	AUDIOMODULATION_API virtual EMovieSceneChannelProxyType CacheChannelProxy() override;
+	AUDIOMODULATION_API virtual void ImportEntityImpl(UMovieSceneEntitySystemLinker* EntityLinker, const FEntityImportParams& Params, FImportedEntity* OutImportedEntity) override;
+
+
+	AUDIOMODULATION_API void SetBus(USoundModulatorBase* NewBus);
+
+	UPROPERTY()
+	FMovieSceneFloatChannel BusValue;
+
+	UPROPERTY()
+	TObjectPtr<USoundModulatorBase> MixBus;
+};
+
diff --git a/Engine/Plugins/Runtime/AudioModulation/Source/AudioModulation/Public/MovieSceneAudioControlBusTrack.h b/Engine/Plugins/Runtime/AudioModulation/Source/AudioModulation/Public/MovieSceneAudioControlBusTrack.h
new file mode 100644
--- /dev/null
+++ b/Engine/Plugins/Runtime/AudioModulation/Source/AudioModulation/Public/MovieSceneAudioControlBusTrack.h
@@ -0,0 +1,43 @@
+// Copyright Epic Games, Inc. All Rights Reserved.
+
+#pragma once
+
+#include "IAudioModulation.h"
+#include "MovieSceneNameableTrack.h"
+#include "MovieSceneSection.h"
+#include "UObject/ObjectMacros.h"
+
+#include "MovieSceneAudioControlBusTrack.generated.h"
+
+UCLASS(MinimalAPI)
+class UMovieSceneAudioControlBusTrack
+	: public UMovieSceneNameableTrack
+{
+	GENERATED_UCLASS_BODY()
+
+public:
+
+	AUDIOMODULATION_API UMovieSceneSection* AddNewControlBus(USoundModulatorBase* Bus, FFrameNumber Time, int32 RowIndex);
+
+	// UMovieSceneTrack interface
+
+	AUDIOMODULATION_API virtual bool SupportsType(TSubclassOf<UMovieSceneSection> SectionClass) const override;
+	AUDIOMODULATION_API virtual void RemoveAllAnimationData() override;
+	AUDIOMODULATION_API virtual bool HasSection(const UMovieSceneSection& Section) const override;
+	AUDIOMODULATION_API virtual void AddSection(UMovieSceneSection& Section) override;
+	AUDIOMODULATION_API virtual void RemoveSection(UMovieSceneSection& Section) override;
+	AUDIOMODULATION_API virtual void RemoveSectionAt(int32 SectionIndex) override;
+	AUDIOMODULATION_API virtual bool IsEmpty() const override;
+	AUDIOMODULATION_API virtual const TArray<UMovieSceneSection*>& GetAllSections() const override;
+	AUDIOMODULATION_API virtual bool SupportsMultipleRows() const override;
+	AUDIOMODULATION_API virtual UMovieSceneSection* CreateNewSection() override;
+
+	AUDIOMODULATION_API virtual void GetCompatibleUserDecorationsImpl(TSet<UClass*>& OutClasses) const override;
+
+private:
+
+	/** List of all root audio sections */
+	UPROPERTY()
+	TArray<TObjectPtr<UMovieSceneSection>> AudioSections;
+};
+
diff --git a/Engine/Plugins/Runtime/AudioModulation/Source/AudioModulation/Public/MovieSceneAudioModulationComponentTypes.h b/Engine/Plugins/Runtime/AudioModulation/Source/AudioModulation/Public/MovieSceneAudioModulationComponentTypes.h
new file mode 100644
--- /dev/null
+++ b/Engine/Plugins/Runtime/AudioModulation/Source/AudioModulation/Public/MovieSceneAudioModulationComponentTypes.h
@@ -0,0 +1,32 @@
+// Copyright Epic Games, Inc. All Rights Reserved.
+
+#pragma once
+
+#include "MovieSceneAudioControlBusSection.h"
+#include "EntitySystem/MovieSceneEntityIDs.h"
+
+#include "MovieSceneAudioModulationComponentTypes.generated.h"
+
+/** Component data for control bus tracks tracks */
+USTRUCT()
+struct FMovieSceneAudioControlBusComponentData
+{
+	GENERATED_BODY()
+
+	UPROPERTY()
+	TObjectPtr<UMovieSceneAudioControlBusSection> Section = nullptr;
+};
+
+struct FMovieSceneAudioModulationComponentTypes
+{
+public:
+	AUDIOMODULATION_API ~FMovieSceneAudioModulationComponentTypes();
+
+	static AUDIOMODULATION_API FMovieSceneAudioModulationComponentTypes* Get();
+	static AUDIOMODULATION_API void Destroy();
+
+	UE::MovieScene::TComponentTypeID<FMovieSceneAudioControlBusComponentData> AudioControlBus;
+
+private:
+	FMovieSceneAudioModulationComponentTypes();
+};
\ No newline at end of file
diff --git a/Engine/Plugins/Runtime/AudioModulation/Source/AudioModulationEditor/AudioModulationEditor.build.cs b/Engine/Plugins/Runtime/AudioModulation/Source/AudioModulationEditor/AudioModulationEditor.build.cs
--- a/Engine/Plugins/Runtime/AudioModulation/Source/AudioModulationEditor/AudioModulationEditor.build.cs
+++ b/Engine/Plugins/Runtime/AudioModulation/Source/AudioModulationEditor/AudioModulationEditor.build.cs
@@ -36,7 +36,19 @@ public class AudioModulationEditor : ModuleRules
 				"ToolWidgets",
 				"TraceAnalysis",
 				"WaveTable",
-				"WaveTableEditor"
+				"WaveTableEditor",
+				"MovieScene",
+				"Sequencer",
+				"SequencerCore",
+				"MovieSceneTools"
+			}
+		);
+
+		PublicDependencyModuleNames.AddRange(
+			new string[]
+			{
+				"Sequencer",
+				"MovieSceneTools"
 			}
 		);
 
diff --git a/Engine/Plugins/Runtime/AudioModulation/Source/AudioModulationEditor/Private/AudioControlBusTrackEditor.cpp b/Engine/Plugins/Runtime/AudioModulation/Source/AudioModulationEditor/Private/AudioControlBusTrackEditor.cpp
new file mode 100644
--- /dev/null
+++ b/Engine/Plugins/Runtime/AudioModulation/Source/AudioModulationEditor/Private/AudioControlBusTrackEditor.cpp
@@ -0,0 +1,258 @@
+// Copyright Epic Games, Inc. All Rights Reserved.
+
+#include "AudioControlBusTrackEditor.h"
+
+#include "AssetRegistry/AssetRegistryModule.h"
+#include "AudioDevice.h"
+#include "ContentBrowserModule.h"
+#include "Editor.h"
+#include "Engine/World.h"
+#include "IContentBrowserSingleton.h"
+#include "ISequencer.h"
+#include "ISequencerTrackEditor.h"
+#include "MVVM/Extensions/ITrackExtension.h"
+#include "MVVM/ViewModels/OutlinerColumns/OutlinerColumnTypes.h"
+#include "MVVM/Views/ViewUtilities.h"
+#include "SequencerSettings.h"
+#include "MovieSceneAudioControlBusTrack.h"
+
+#define LOCTEXT_NAMESPACE "FAudioControlBusTrackEditor"
+
+FAudioControlBusTrackEditor::FAudioControlBusTrackEditor(TSharedRef<ISequencer> InSequencer)
+	: FMovieSceneTrackEditor(InSequencer)
+{
+}
+
+FAudioControlBusTrackEditor::~FAudioControlBusTrackEditor()
+{
+}
+
+TSharedRef<ISequencerTrackEditor> FAudioControlBusTrackEditor::CreateTrackEditor(TSharedRef<ISequencer> OwningSequencer)
+{
+	return MakeShareable(new FAudioControlBusTrackEditor(OwningSequencer));
+}
+
+FText FAudioControlBusTrackEditor::GetDisplayName() const
+{
+	return LOCTEXT("AudioControlBusTrackEditor_DisplayName", "Control Bus");
+}
+
+void FAudioControlBusTrackEditor::BuildAddTrackMenu(FMenuBuilder& MenuBuilder)
+{
+	MenuBuilder.AddMenuEntry(
+		LOCTEXT("AddTrack", "Control Bus Track"),
+		LOCTEXT("AddTooltip", "Adds a new control track that can be used for mixing"),
+		FSlateIcon(FAppStyle::GetAppStyleSetName(), "Sequencer.Tracks.Audio"),
+		FUIAction(
+			FExecuteAction::CreateRaw(this, &FAudioControlBusTrackEditor::HandleAddAudioControlBusTrackMenuEntryExecute)
+		)
+	);
+}
+
+TSharedPtr<SWidget> FAudioControlBusTrackEditor::BuildOutlinerColumnWidget(const FBuildColumnWidgetParams& Params, const FName& ColumnName)
+{
+	using namespace UE::Sequencer;
+
+	UMovieSceneTrack* Track = Params.TrackModel->GetTrack();
+	if (!Track)
+	{
+		return SNullWidget::NullWidget;
+	}
+
+	if (ColumnName == FCommonOutlinerNames::Add)
+	{
+		return UE::Sequencer::MakeAddButton(
+			LOCTEXT("AudioControlBusText", "Control Bus"),
+			FOnGetContent::CreateSP(
+				this, &FAudioControlBusTrackEditor::BuildAudioControlBusSubMenu,
+				FOnAssetSelected::CreateRaw(this, &FAudioControlBusTrackEditor::OnAudioControlBusAssetSelected, Track),
+				FOnAssetEnterPressed::CreateRaw(this, &FAudioControlBusTrackEditor::OnAudioControlBusAssetEnterPressed, Track)
+			),
+			Params.ViewModel
+		);
+	}
+
+	return FMovieSceneTrackEditor::BuildOutlinerColumnWidget(Params, ColumnName);
+}
+
+bool FAudioControlBusTrackEditor::SupportsType(TSubclassOf<UMovieSceneTrack> Type) const
+{
+	return Type == UMovieSceneAudioControlBusTrack::StaticClass();
+}
+
+bool FAudioControlBusTrackEditor::SupportsSequence(UMovieSceneSequence* InSequence) const
+{
+	return true;
+}
+
+bool FAudioControlBusTrackEditor::IsResizable(UMovieSceneTrack* InTrack) const
+{
+	return true;
+}
+
+void FAudioControlBusTrackEditor::HandleAddAudioControlBusTrackMenuEntryExecute()
+{
+	if (!IsModulationEnabled())
+	{
+		return;
+	}
+
+	UMovieScene* FocusedMovieScene = GetFocusedMovieScene();
+
+	if (FocusedMovieScene == nullptr)
+	{
+		return;
+	}
+
+	// If movie scene is readonly we cannot add to it, user will have to make it writable and select the add control bus track option again
+	if (FocusedMovieScene->IsReadOnly())
+	{
+		return;
+	}
+
+	const FScopedTransaction Transaction(NSLOCTEXT("Sequencer", "AddAudioControlBusTrack_Transaction", "Add Control Bus Track"));
+	FocusedMovieScene->Modify();
+
+	UMovieSceneAudioControlBusTrack* NewTrack = FocusedMovieScene->AddTrack<UMovieSceneAudioControlBusTrack>();
+	ensure(NewTrack);
+
+	NewTrack->SetDisplayName(LOCTEXT("AudioControlBusTrackName", "Control Bus"));
+
+	if (GetSequencer().IsValid())
+	{
+		GetSequencer()->OnAddTrack(NewTrack, FGuid());
+	}
+}
+
+void FAudioControlBusTrackEditor::OnAudioControlBusAssetSelected(const FAssetData& AssetData, UMovieSceneTrack* Track)
+{
+	FSlateApplication::Get().DismissAllMenus();
+
+	if (!IsModulationEnabled())
+	{
+		return;
+	}
+
+	UObject* SelectedObject = AssetData.GetAsset();
+
+	if (SelectedObject)
+	{
+		USoundModulatorBase* NewBus = CastChecked<USoundModulatorBase>(AssetData.GetAsset());
+		if (NewBus != nullptr)
+		{
+			const FScopedTransaction Transaction(NSLOCTEXT("Sequencer", "AddAudioControlBus_Transaction", "Add Control Bus"));
+
+			UMovieSceneAudioControlBusTrack* AudioControlBusTrack = Cast<UMovieSceneAudioControlBusTrack>(Track);
+			check(AudioControlBusTrack);
+
+			AudioControlBusTrack->Modify();
+
+			check(GetSequencer().IsValid());
+			
+			FFrameTime KeyTime = GetSequencer()->GetLocalTime().Time;
+
+			UMovieSceneSection* NewSection = AudioControlBusTrack->AddNewControlBus(NewBus, KeyTime.FrameNumber, INDEX_NONE);
+
+			GetSequencer()->EmptySelection();
+			GetSequencer()->SelectSection(NewSection);
+			GetSequencer()->ThrobSectionSelection();
+
+			GetSequencer()->NotifyMovieSceneDataChanged(EMovieSceneDataChangeType::MovieSceneStructureItemAdded);
+		}
+	}
+}
+
+void FAudioControlBusTrackEditor::OnAudioControlBusAssetEnterPressed(const TArray<FAssetData>& AssetData, UMovieSceneTrack* Track)
+{
+	// By default enter will select the first list element if there is at least 1 element
+	if (AssetData.Num() > 0)
+	{
+		OnAudioControlBusAssetSelected(AssetData[0].GetAsset(), Track);
+	}
+}
+
+TSharedRef<SWidget> FAudioControlBusTrackEditor::BuildAudioControlBusSubMenu(FOnAssetSelected OnAssetSelected, FOnAssetEnterPressed OnAssetEnterPressed)
+{
+	FMenuBuilder MenuBuilder(true, nullptr);
+
+	if (!IsModulationEnabled())
+	{
+		return MenuBuilder.MakeWidget();
+	}
+
+	const TSharedPtr<ISequencer> SequencerPtr = GetSequencer();
+	const UMovieSceneSequence* Sequence = SequencerPtr.IsValid() ? SequencerPtr->GetFocusedMovieSceneSequence() : nullptr;
+
+	FAssetRegistryModule& AssetRegistryModule = FModuleManager::LoadModuleChecked<FAssetRegistryModule>(TEXT("AssetRegistry"));
+	TArray<FTopLevelAssetPath> ClassNames;
+
+	// Due to the plugin storing the USoundControlBus class we have to request the plugin to fetch all the control bus assets since several asset types inherit from USoundModulatorBase
+	if (GEditor)
+	{
+		if (UWorld* World = GEditor->GetEditorWorldContext().World())
+		{
+			FAudioDeviceHandle AudioDeviceHandle = World->GetAudioDevice();
+			if (AudioDeviceHandle.IsValid() && AudioDeviceHandle->IsModulationPluginEnabled())
+			{
+				if (IAudioModulationManager* Modulation = AudioDeviceHandle->ModulationInterface.Get())
+				{
+					Modulation->GetControlBusAssetPath(ClassNames);
+				}
+			}
+		}
+	}
+
+	TSet<FTopLevelAssetPath> DerivedClassNames;
+	AssetRegistryModule.Get().GetDerivedClassNames(ClassNames, TSet<FTopLevelAssetPath>(), DerivedClassNames);
+
+	FAssetPickerConfig AssetPickerConfig;
+	{
+		AssetPickerConfig.OnAssetSelected = OnAssetSelected;
+		AssetPickerConfig.OnAssetEnterPressed = OnAssetEnterPressed;
+		AssetPickerConfig.bAllowNullSelection = false;
+		AssetPickerConfig.bAddFilterUI = true;
+		AssetPickerConfig.bShowTypeInColumnView = false;
+		AssetPickerConfig.InitialAssetViewType = EAssetViewType::List;
+		for (FTopLevelAssetPath ClassName : DerivedClassNames)
+		{
+			AssetPickerConfig.Filter.ClassPaths.Add(ClassName);
+		}
+		AssetPickerConfig.SaveSettingsName = TEXT("SequencerAssetPicker");
+		AssetPickerConfig.AdditionalReferencingAssets.Add(FAssetData(Sequence));
+	}
+
+	FContentBrowserModule& ContentBrowserModule = FModuleManager::Get().LoadModuleChecked<FContentBrowserModule>(TEXT("ContentBrowser"));
+
+	const float WidthOverride = SequencerPtr.IsValid() ? SequencerPtr->GetSequencerSettings()->GetAssetBrowserWidth() : 500.f;
+	const float HeightOverride = SequencerPtr.IsValid() ? SequencerPtr->GetSequencerSettings()->GetAssetBrowserHeight() : 400.f;
+
+	TSharedPtr<SBox> MenuEntry = SNew(SBox)
+		.WidthOverride(WidthOverride)
+		.HeightOverride(HeightOverride)
+		[
+			ContentBrowserModule.Get().CreateAssetPicker(AssetPickerConfig)
+		];
+
+	MenuBuilder.AddWidget(MenuEntry.ToSharedRef(), FText::GetEmpty(), true);
+
+	return MenuBuilder.MakeWidget();
+}
+
+bool FAudioControlBusTrackEditor::IsModulationEnabled()
+{
+	if (GEditor)
+	{
+		if (UWorld* World = GEditor->GetEditorWorldContext().World())
+		{
+			FAudioDeviceHandle AudioDeviceHandle = World->GetAudioDevice();
+			if (AudioDeviceHandle.IsValid())
+			{
+				return AudioDeviceHandle->IsModulationPluginEnabled();
+			}
+		}
+	}
+
+	return false;
+}
+
+#undef LOCTEXT_NAMESPACE
diff --git a/Engine/Plugins/Runtime/AudioModulation/Source/AudioModulationEditor/Private/AudioControlBusTrackEditor.h b/Engine/Plugins/Runtime/AudioModulation/Source/AudioModulationEditor/Private/AudioControlBusTrackEditor.h
new file mode 100644
--- /dev/null
+++ b/Engine/Plugins/Runtime/AudioModulation/Source/AudioModulationEditor/Private/AudioControlBusTrackEditor.h
@@ -0,0 +1,62 @@
+// Copyright Epic Games, Inc. All Rights Reserved.
+
+#pragma once
+
+#include "IAudioModulation.h"
+#include "IContentBrowserSingleton.h"
+#include "ISequencer.h"
+#include "ISequencerTrackEditor.h"
+#include "MovieSceneSection.h"
+#include "MovieSceneTrackEditor.h"
+
+#define UE_API AUDIOMODULATIONEDITOR_API
+
+class FAudioControlBusTrackEditor
+	: public FMovieSceneTrackEditor
+{
+public:
+	/**
+	* Constructor
+	*
+	* @param InSequencer The sequencer instance to be used by this tool
+	*/
+	UE_API FAudioControlBusTrackEditor(TSharedRef<ISequencer> InSequencer);
+
+	/** Virtual destructor. */
+	UE_API virtual ~FAudioControlBusTrackEditor();
+
+	/**
+	 * Creates an instance of this class.  Called by a sequencer
+	 *
+	 * @param OwningSequencer The sequencer instance to be used by this tool
+	 * @return The new instance of this class
+	 */
+	static UE_API TSharedRef<ISequencerTrackEditor> CreateTrackEditor(TSharedRef<ISequencer> OwningSequencer);
+
+	// ISequencerTrackEditor interface
+
+	UE_API virtual FText GetDisplayName() const override;
+	UE_API virtual void BuildAddTrackMenu(FMenuBuilder& MenuBuilder) override;
+	UE_API virtual TSharedPtr<SWidget> BuildOutlinerColumnWidget(const FBuildColumnWidgetParams& Params, const FName& ColumnName) override;
+	UE_API virtual bool SupportsType(TSubclassOf<UMovieSceneTrack> Type) const override;
+	UE_API virtual bool SupportsSequence(UMovieSceneSequence* InSequence) const override;
+	UE_API virtual bool IsResizable(UMovieSceneTrack* InTrack) const override;
+
+private:
+
+	/** Callback for executing the "Add Audio Control Bus Track" menu entry. */
+	UE_API void HandleAddAudioControlBusTrackMenuEntryExecute();
+
+	/** Audio control bus asset selected */
+	UE_API void OnAudioControlBusAssetSelected(const FAssetData& AssetData, UMovieSceneTrack* Track);
+
+	/** Audio control bus asset enter pressed */
+	UE_API void OnAudioControlBusAssetEnterPressed(const TArray<FAssetData>& AssetData, UMovieSceneTrack* Track);
+
+	/** Audio control bus sub menu */
+	UE_API TSharedRef<SWidget> BuildAudioControlBusSubMenu(FOnAssetSelected OnAssetSelected, FOnAssetEnterPressed OnAssetEnterPressed);
+
+	UE_API bool IsModulationEnabled();
+};
+
+#undef UE_API
\ No newline at end of file
diff --git a/Engine/Plugins/Runtime/AudioModulation/Source/AudioModulationEditor/Private/AudioModulationEditor.cpp b/Engine/Plugins/Runtime/AudioModulation/Source/AudioModulationEditor/Private/AudioModulationEditor.cpp
--- a/Engine/Plugins/Runtime/AudioModulation/Source/AudioModulationEditor/Private/AudioModulationEditor.cpp
+++ b/Engine/Plugins/Runtime/AudioModulation/Source/AudioModulationEditor/Private/AudioModulationEditor.cpp
@@ -25,12 +25,15 @@
 #include "Layouts/SoundModulationTransformLayout.h"
 #include "LevelEditor.h"
 #include "Modules/ModuleManager.h"
+#include "MovieSceneToolsModule.h"
 #include "Sound/SoundBase.h"
 #include "SoundModulationParameter.h"
 #include "SoundModulationPatch.h"
 #include "Templates/SharedPointer.h"
 #include "Textures/SlateIcon.h"
 #include "UObject/UObjectIterator.h"
+#include "ISequencerModule.h"
+#include "AudioControlBusTrackEditor.h"
 
 #if WITH_AUDIOMODULATION_METASOUND_SUPPORT
 #include "MetasoundDataReference.h"
@@ -155,6 +158,16 @@ void FAudioModulationEditorModule::StartupModule()
 
 	FSlateStyleRegistry::RegisterSlateStyle(*StyleSet.Get());
 
+	using namespace UE::MovieScene;
+
+#if WITH_EDITOR
+	if (GIsEditor)
+	{
+		ISequencerModule& SequencerModule = FModuleManager::Get().LoadModuleChecked<ISequencerModule>("Sequencer");
+		AudioControlBusTrackCreateEditorHandle = SequencerModule.RegisterTrackEditor(FOnCreateTrackEditor::CreateStatic(&FAudioControlBusTrackEditor::CreateTrackEditor));
+	}
+#endif
+
 #if WITH_AUDIOMODULATION_METASOUND_SUPPORT
 	{
 		using namespace Metasound;
@@ -235,6 +248,14 @@ void FAudioModulationEditorModule::ShutdownModule()
 
 	FSlateStyleRegistry::UnRegisterSlateStyle(*StyleSet.Get());
 
+#if WITH_EDITOR
+	if (GIsEditor)
+	{
+		ISequencerModule& SequencerModule = FModuleManager::Get().LoadModuleChecked<ISequencerModule>("Sequencer");
+		SequencerModule.UnRegisterTrackEditor(AudioControlBusTrackCreateEditorHandle);
+	}
+#endif
+
 	::UnregisterModulationInAudioInsights();
 }
 
diff --git a/Engine/Plugins/Runtime/AudioModulation/Source/AudioModulationEditor/Private/AudioModulationEditor.h b/Engine/Plugins/Runtime/AudioModulation/Source/AudioModulationEditor/Private/AudioModulationEditor.h
--- a/Engine/Plugins/Runtime/AudioModulation/Source/AudioModulationEditor/Private/AudioModulationEditor.h
+++ b/Engine/Plugins/Runtime/AudioModulation/Source/AudioModulationEditor/Private/AudioModulationEditor.h
@@ -35,4 +35,6 @@ private:
 	TSharedPtr<FExtensibilityManager> ModulationPatchToolBarExtensibilityManager;
 
 	TSharedPtr<FSlateStyleSet> StyleSet;
+
+	FDelegateHandle AudioControlBusTrackCreateEditorHandle;
 };
diff --git a/Engine/Source/Runtime/AudioExtensions/Public/IAudioModulation.h b/Engine/Source/Runtime/AudioExtensions/Public/IAudioModulation.h
--- a/Engine/Source/Runtime/AudioExtensions/Public/IAudioModulation.h
+++ b/Engine/Source/Runtime/AudioExtensions/Public/IAudioModulation.h
@@ -187,6 +187,26 @@ public:
 
 	virtual void OnAuditionEnd() = 0;
 
+	virtual void GetControlBusAssetPath(TArray<FTopLevelAssetPath>& OutClassNames) = 0;
+
+	virtual void CreateNewSequencerBusMix(
+		const UObject* WorldContextObject,
+		FName Name,
+		const TArray<USoundModulatorBase*>& Buses,
+		float Value = 1.0f,
+		float AttackTime = 0.1f,
+		float ReleaseTime = 0.1f,
+		bool bActivate = true) = 0;
+
+	virtual void UpdateMixByModulator(
+		const UObject* WorldContextObject,
+		FName Mix,
+		USoundModulatorBase& Modulator,
+		float Value = 1.0f,
+		float FadeTime = -1.0f) = 0;
+
+	virtual void ClearSequencerBusMix(FName Name) = 0;
+
 #if !UE_BUILD_SHIPPING
 	/** Request to post help from active plugin (non-shipping builds only) */
 	virtual bool OnPostHelp(FCommonViewportClient* ViewportClient, const TCHAR* Stream) = 0; 
diff --git a/Engine/Source/Runtime/MovieSceneTracks/MovieSceneTracks.Build.cs b/Engine/Source/Runtime/MovieSceneTracks/MovieSceneTracks.Build.cs
--- a/Engine/Source/Runtime/MovieSceneTracks/MovieSceneTracks.Build.cs
+++ b/Engine/Source/Runtime/MovieSceneTracks/MovieSceneTracks.Build.cs
@@ -14,7 +14,6 @@ public class MovieSceneTracks : ModuleRules
 				"MovieScene",
 				"TimeManagement",
 				"AnimationCore",
-				"AudioExtensions",
 				"Constraints",
 				"AudioMixer"
 			});
