Sequencer: Control Bus track feedback

diff --git a/Engine/Plugins/Runtime/AudioModulation/Source/AudioModulation/Private/MovieSceneAudioControlBusTrack.cpp b/Engine/Plugins/Runtime/AudioModulation/Source/AudioModulation/Private/MovieSceneAudioControlBusTrack.cpp
--- a/Engine/Plugins/Runtime/AudioModulation/Source/AudioModulation/Private/MovieSceneAudioControlBusTrack.cpp
+++ b/Engine/Plugins/Runtime/AudioModulation/Source/AudioModulation/Private/MovieSceneAudioControlBusTrack.cpp
@@ -3,6 +3,7 @@
 #include "MovieSceneAudioControlBusTrack.h"
 
 #include "Decorations/MovieSceneSectionAnchorsDecoration.h"
+#include "MovieScene.h"
 #include "MovieSceneAudioControlBusSection.h"
 
 #include UE_INLINE_GENERATED_CPP_BY_NAME(MovieSceneAudioControlBusTrack)
@@ -20,7 +21,8 @@ UMovieSceneSection* UMovieSceneAudioControlBusTrack::AddNewControlBus(USoundModu
 	UMovieSceneAudioControlBusSection* NewSection = Cast<UMovieSceneAudioControlBusSection>(CreateNewSection());
 
 	NewSection->SetBus(Bus);
-	constexpr float DefaultDuration = 5.0f;
+	FFrameRate FrameRate = GetTypedOuter<UMovieScene>()->GetTickResolution();
+	const float DefaultDuration = 5.0f * FrameRate.AsDecimal();
 	NewSection->InitialPlacementOnRow(AudioSections, Time, DefaultDuration, RowIndex);
 	AudioSections.Add(NewSection);
 
