Add Description field to Control Bus

diff --git a/Engine/Plugins/Runtime/AudioModulation/Source/AudioModulation/Public/SoundControlBus.h b/Engine/Plugins/Runtime/AudioModulation/Source/AudioModulation/Public/SoundControlBus.h
--- a/Engine/Plugins/Runtime/AudioModulation/Source/AudioModulation/Public/SoundControlBus.h
+++ b/Engine/Plugins/Runtime/AudioModulation/Source/AudioModulation/Public/SoundControlBus.h
@@ -34,6 +34,9 @@ public:
 	/** If true, Address field is used in place of object name for address used when applying mix changes using filtering. */
 	UPROPERTY(EditAnywhere, Category = "Advanced|Mix", BlueprintReadWrite)
 	bool bOverrideAddress;
+	
+	UPROPERTY(EditAnywhere, Category = General, meta=(MultiLine=true))
+	FString Description;
 #endif // WITH_EDITORONLY_DATA
 
 	/** Address to use when applying mix changes. */
diff --git a/Engine/Plugins/Runtime/AudioModulation/Source/AudioModulationEditor/Private/AssetTypeActions/AssetTypeActions_SoundControlBus.cpp b/Engine/Plugins/Runtime/AudioModulation/Source/AudioModulationEditor/Private/AssetTypeActions/AssetTypeActions_SoundControlBus.cpp
--- a/Engine/Plugins/Runtime/AudioModulation/Source/AudioModulationEditor/Private/AssetTypeActions/AssetTypeActions_SoundControlBus.cpp
+++ b/Engine/Plugins/Runtime/AudioModulation/Source/AudioModulationEditor/Private/AssetTypeActions/AssetTypeActions_SoundControlBus.cpp
@@ -6,6 +6,15 @@
 
 #define LOCTEXT_NAMESPACE "AssetTypeActions"
 
+FText FAssetTypeActions_SoundControlBus::GetAssetDescription(const FAssetData& AssetData) const
+{
+	if (USoundControlBus* Bus = Cast<USoundControlBus>(AssetData.GetAsset()))
+	{
+		return FText::FromString(Bus->Description);
+	}
+	return FText::GetEmpty();
+}
+
 UClass* FAssetTypeActions_SoundControlBus::GetSupportedClass() const
 {
 	return USoundControlBus::StaticClass();
diff --git a/Engine/Plugins/Runtime/AudioModulation/Source/AudioModulationEditor/Private/AssetTypeActions/AssetTypeActions_SoundControlBus.h b/Engine/Plugins/Runtime/AudioModulation/Source/AudioModulationEditor/Private/AssetTypeActions/AssetTypeActions_SoundControlBus.h
--- a/Engine/Plugins/Runtime/AudioModulation/Source/AudioModulationEditor/Private/AssetTypeActions/AssetTypeActions_SoundControlBus.h
+++ b/Engine/Plugins/Runtime/AudioModulation/Source/AudioModulationEditor/Private/AssetTypeActions/AssetTypeActions_SoundControlBus.h
@@ -11,6 +11,7 @@ class FAssetTypeActions_SoundControlBus : public FAssetTypeActions_Base
 public:
 	// IAssetTypeActions Implementation
 	virtual FText GetName() const override { return NSLOCTEXT("AssetTypeActions", "AssetTypeActions_SoundControlBus", "Control Bus"); }
+	virtual FText GetAssetDescription(const FAssetData& AssetData) const override;
 	virtual FColor GetTypeColor() const override { return UAudioModulationStyle::GetControlBusColor(); }
 	virtual UClass* GetSupportedClass() const override;
 	virtual uint32 GetCategories() override { return EAssetTypeCategories::Sounds; }
