[AudioModulation] Allow exposure of debugger data through modular feature

diff --git a/Engine/Plugins/Runtime/AudioModulation/Source/AudioModulation/Private/AudioModulationDebugger.cpp b/Engine/Plugins/Runtime/AudioModulation/Source/AudioModulation/Private/AudioModulationDebugger.cpp
--- a/Engine/Plugins/Runtime/AudioModulation/Source/AudioModulation/Private/AudioModulationDebugger.cpp
+++ b/Engine/Plugins/Runtime/AudioModulation/Source/AudioModulation/Private/AudioModulationDebugger.cpp
@@ -14,13 +14,13 @@
 #include "Engine/Engine.h"
 #include "HAL/IConsoleManager.h"
 #include "IAudioModulation.h"
+#include "Features/IModularFeatures.h"
 #include "Misc/CoreDelegates.h"
 #include "SoundControlBusProxy.h"
 #include "SoundModulationPatch.h"
 #include "SoundModulationProxy.h"
 #include "SoundModulationValue.h"
 
-
 namespace AudioModulation
 {
 	namespace Debug
@@ -448,6 +448,24 @@ namespace AudioModulation
 			return A.Name < B.Name;
 		}
 
+		template <typename T>
+		void FilterDebugInfoArray(const TArray<T>& InArray, const FString& FilterString, int32 MaxCount, TArray<T>& FilteredArray)
+		{
+			int32 FilteredItemCount = 0;
+			for (const T& Item : InArray)
+			{
+				const bool Filtered = !FilterString.IsEmpty() && !Item.Name.Contains(FilterString);
+				if (!Filtered)
+				{
+					FilteredArray.Add(Item);
+					if (++FilteredItemCount >= MaxCount)
+					{
+						return;
+					}
+				}
+			}
+		}
+
 		template <typename T, typename U>
 		void FilterDebugArray(const T& Map, const FString& FilterString, int32 MaxCount, TArray<const U*>& FilteredArray)
 		{
@@ -498,6 +516,14 @@ namespace AudioModulation
 		, bShowGlobalMixes(1)
 		, ElapsedSinceLastUpdate(0.0f)
 	{
+		IModularFeatures::FScopedLockModularFeatureList ModularFeatureListLock;
+		IModularFeatures::Get().RegisterModularFeature(DebugDataModularFeatureName, this);
+	}
+
+	FAudioModulationDebugger::~FAudioModulationDebugger()
+	{
+		IModularFeatures::FScopedLockModularFeatureList ModularFeatureListLock;
+		IModularFeatures::Get().UnregisterModularFeature(DebugDataModularFeatureName, this);
 	}
 
 	void FAudioModulationDebugger::UpdateDebugData(double InElapsed, const FReferencedProxies& InRefProxies, Audio::FDeviceId DeviceId)
@@ -515,23 +541,30 @@ namespace AudioModulation
 		}
 		ElapsedSinceLastUpdate = 0.0f;
 
-		static const int32 MaxFilteredBuses = 8;
-		TArray<const FControlBusProxy*> FilteredBusProxies;
-		Debug::FilterDebugArray<FBusProxyMap, FControlBusProxy>(InRefProxies.Buses, BusStringFilter, MaxFilteredBuses, FilteredBusProxies);
-		TArray<FControlBusDebugInfo> RefreshedFilteredBuses;
-		for (const FControlBusProxy* Proxy : FilteredBusProxies)
+		//Create Bus Debug Info Arrays (filtered and unfiltered)
+		TArray<FControlBusDebugInfo> RefreshedUnfilteredBuses;
+		for (const TTuple<FBusId, FControlBusProxy>& BusProxyMap : InRefProxies.Buses)
 		{
+			const FControlBusProxy& Proxy = BusProxyMap.Value;
 			FControlBusDebugInfo DebugInfo;
-			DebugInfo.DefaultValue = Proxy->GetDefaultValue();
-			DebugInfo.Id = Proxy->GetId();
-			DebugInfo.GeneratorValue = Proxy->GetGeneratorValue();
-			DebugInfo.MixValue = Proxy->GetMixValue();
-			DebugInfo.Name = Proxy->GetName().ToString();
-			DebugInfo.RefCount = Proxy->GetRefCount();
-			DebugInfo.Value = Proxy->GetValue();
-			RefreshedFilteredBuses.Add(DebugInfo);
+			DebugInfo.DefaultValue = Proxy.GetDefaultValue();
+			DebugInfo.Id = Proxy.GetId();
+			DebugInfo.GeneratorValue = Proxy.GetGeneratorValue();
+			DebugInfo.MixValue = Proxy.GetMixValue();
+			DebugInfo.Name = Proxy.GetName().ToString();
+			DebugInfo.ParameterName = Proxy.GetParameterName().ToString();
+			DebugInfo.RefCount = Proxy.GetRefCount();
+			DebugInfo.Value = Proxy.GetValue();
+			RefreshedUnfilteredBuses.Add(DebugInfo);
 		}
 
+		RefreshedUnfilteredBuses.Sort(&Debug::CompareNames<FControlBusDebugInfo>);
+
+		TArray<FControlBusDebugInfo> RefreshedFilteredBuses;
+		static const int32 MaxFilteredBuses = 8;
+		Debug::FilterDebugInfoArray<FControlBusDebugInfo>(RefreshedUnfilteredBuses, BusStringFilter, MaxFilteredBuses,RefreshedFilteredBuses);
+
+		
 		FAudioModulationSystem* ModulationSystem = nullptr;
 		if (FAudioDeviceManager* DeviceManager = FAudioDeviceManager::Get())
 		{
@@ -547,37 +580,43 @@ namespace AudioModulation
 			}
 		}
 
+		//Create Bus Mix Debug Info Arrays (filtered and unfiltered)
+		TArray<FControlBusMixDebugInfo> RefreshedUnfilteredMixes;
 		TArray<FControlBusMixDebugInfo> RefreshedFilteredMixes;
+
 		if (ModulationSystem)
 		{
-			static const int32 MaxFilteredMixes = 16;
-			TMap<FBusMixId, const FModulatorBusMixProxy*> FilteredMixProxies;
-			Debug::FilterDebugMap<FBusMixProxyMap, FModulatorBusMixProxy, FBusMixId>(InRefProxies.BusMixes, MixStringFilter, MaxFilteredMixes, FilteredMixProxies);
-			for (TPair<FBusMixId, const FModulatorBusMixProxy*>& Pair : FilteredMixProxies)
+			TMap<FBusMixId, const FModulatorBusMixProxy*> UnfilteredMixProxies;
+			for (const TTuple<FBusMixId, FModulatorBusMixProxy>& Pair : InRefProxies.BusMixes)
 			{
-				const FModulatorBusMixProxy* Proxy = Pair.Value;
+				const FModulatorBusMixProxy& Proxy = Pair.Value;
 			
 				FControlBusMixDebugInfo DebugInfo;
-				DebugInfo.Name = Proxy->GetName().ToString();
-				DebugInfo.RefCount = Proxy->GetRefCount();
+				DebugInfo.Name = Proxy.GetName().ToString();
+				DebugInfo.RefCount = Proxy.GetRefCount();
 				DebugInfo.Id = Pair.Key;
-				DebugInfo.Timer = Proxy->GetTimeRemaining();
+				DebugInfo.Timer = Proxy.GetTimeRemaining();
 				// Bus mix is global if it only has one stage and its ID is the Control Bus's ID.
-				if (Proxy->Stages.Num() == 1)
+				if (Proxy.Stages.Num() == 1)
 				{
-					uint32 BusId = Proxy->Stages.FindArbitraryElement()->Key;
+					uint32 BusId = Proxy.Stages.FindArbitraryElement()->Key;
 					DebugInfo.bIsGlobal = ModulationSystem->IsBusMixGlobal(BusId);
 				}
 			
-				for (const TPair<FBusId, FModulatorBusMixStageProxy>& Stage : Proxy->Stages)
+				for (const TPair<FBusId, FModulatorBusMixStageProxy>& Stage : Proxy.Stages)
 				{
 					FControlBusMixStageDebugInfo StageDebugInfo;
 					StageDebugInfo.CurrentValue = Stage.Value.Value.GetCurrentValue();
 					StageDebugInfo.TargetValue = Stage.Value.Value.TargetValue;
 					DebugInfo.Stages.Add(Stage.Key, StageDebugInfo);
 				}
-				RefreshedFilteredMixes.Add(DebugInfo);
+				RefreshedUnfilteredMixes.Add(DebugInfo);
 			}
+
+			RefreshedUnfilteredMixes.Sort(&Debug::CompareNames<FControlBusMixDebugInfo>);
+			static const int32 MaxFilteredMixes = 16;
+			Debug::FilterDebugInfoArray<FControlBusMixDebugInfo>(RefreshedUnfilteredMixes, MixStringFilter, MaxFilteredMixes,RefreshedFilteredMixes);
+			
 		}
 		
 		static const int32 MaxFilteredGenerators = 8;
@@ -593,7 +632,7 @@ namespace AudioModulation
 			DebugInfo.Categories = Proxy->GetDebugCategories();
 		}
 
-		FFunctionGraphTask::CreateAndDispatchWhenReady([DebuggerPtr = AsWeak(), RefreshedFilteredBuses, RefreshedFilteredMixes, RefreshedFilteredGenerators]()
+		FFunctionGraphTask::CreateAndDispatchWhenReady([DebuggerPtr = AsWeak(), RefreshedUnfilteredBuses, RefreshedUnfilteredMixes, RefreshedFilteredBuses, RefreshedFilteredMixes, RefreshedFilteredGenerators, DeviceId]()
 		{
 			TSharedPtr<FAudioModulationDebugger> ThisDebugger = DebuggerPtr.Pin();
 			if (!ThisDebugger.IsValid())
@@ -601,11 +640,13 @@ namespace AudioModulation
 				return;
 			}
 
+			ThisDebugger->AssociatedDeviceId = DeviceId;	
+
 			ThisDebugger->FilteredBuses = RefreshedFilteredBuses;
-			ThisDebugger->FilteredBuses.Sort(&Debug::CompareNames<FControlBusDebugInfo>);
+			ThisDebugger->UnfilteredBuses = RefreshedUnfilteredBuses;
 
 			ThisDebugger->FilteredMixes = RefreshedFilteredMixes;
-			ThisDebugger->FilteredMixes.Sort(&Debug::CompareNames<FControlBusMixDebugInfo>);
+			ThisDebugger->UnfilteredMixes = RefreshedUnfilteredMixes;
 
 			FGeneratorSortMap NewFilteredMap;
 			if (ThisDebugger->bShowGenerators)
@@ -858,6 +899,43 @@ namespace AudioModulation
 		}
 	}
 
+	Audio::FDeviceId FAudioModulationDebugger::GetAssociatedDeviceID()
+	{
+		check(IsInGameThread());
+		
+		if (!bActive)
+		{
+			bActive = true;
+		}
+		
+		return AssociatedDeviceId;
+	}
+
+	void FAudioModulationDebugger::GetControlBusDebugInfo(TArray<FControlBusDebugInfo>& OutDebugInfo)
+	{
+		check(IsInGameThread());
+		
+		if (!bActive)
+		{
+			bActive = true;
+			return;
+		}
+		
+		OutDebugInfo = UnfilteredBuses;
+	}
+
+	void FAudioModulationDebugger::GetControlBusMixDebugInfo(TArray<FControlBusMixDebugInfo>& OutDebugInfo)
+	{
+		check(IsInGameThread());
+		
+		if (!bActive)
+		{
+			bActive = true;
+			return;
+		}
+		
+		OutDebugInfo = UnfilteredMixes;
+	}
 } // namespace AudioModulation
 #endif // !UE_BUILD_SHIPPING
 #endif // WITH_AUDIOMODULATION
diff --git a/Engine/Plugins/Runtime/AudioModulation/Source/AudioModulation/Private/AudioModulationDebugger.h b/Engine/Plugins/Runtime/AudioModulation/Source/AudioModulation/Private/AudioModulationDebugger.h
--- a/Engine/Plugins/Runtime/AudioModulation/Source/AudioModulation/Private/AudioModulationDebugger.h
+++ b/Engine/Plugins/Runtime/AudioModulation/Source/AudioModulation/Private/AudioModulationDebugger.h
@@ -5,6 +5,7 @@
 
 #if WITH_AUDIOMODULATION
 #if !UE_BUILD_SHIPPING
+#include "IAudioModulationDebugDataProvider.h"
 #include "SoundControlBus.h"
 #include "SoundControlBusMix.h"
 #include "SoundModulationGenerator.h"
@@ -16,35 +17,7 @@ namespace AudioModulation
 {
 	// Forward Declarations
 	struct FReferencedProxies;
-
-	struct FControlBusMixStageDebugInfo
-	{
-		float TargetValue;
-		float CurrentValue;
-	};
-
-	struct FControlBusMixDebugInfo
-	{
-		FString Name;
-		uint32 Id = -1;
-		uint32 RefCount = 0;
-		bool bIsGlobal = false;
-		double Timer = -1.0;
-
-		TMap<uint32, FControlBusMixStageDebugInfo> Stages;
-	};
-
-	struct FControlBusDebugInfo
-	{
-		FString Name;
-		float DefaultValue;
-		float GeneratorValue;
-		float MixValue;
-		float Value;
-		uint32 Id;
-		uint32 RefCount;
-	};
-
+	
 	struct FGeneratorSort
 	{
 		FORCEINLINE bool operator()(const FString& A, const FString& B) const
@@ -73,10 +46,11 @@ namespace AudioModulation
 		TArray<FInstanceValues> FilteredInstances;
 	};
 
-	class FAudioModulationDebugger : public TSharedFromThis<FAudioModulationDebugger, ESPMode::ThreadSafe>
+	class FAudioModulationDebugger : public IDebugDataProvider, public TSharedFromThis<FAudioModulationDebugger, ESPMode::ThreadSafe>
 	{
 	public:
 		FAudioModulationDebugger();
+		virtual ~FAudioModulationDebugger();
 
 		void UpdateDebugData(double InElapsed, const FReferencedProxies& InRefProxies, Audio::FDeviceId DeviceId);
 		void SetDebugBusFilter(const FString* InNameFilter);
@@ -93,6 +67,12 @@ namespace AudioModulation
 
 		void ResetGeneratorStats();
 
+		//IdebugDataProviderInterface
+		virtual Audio::FDeviceId GetAssociatedDeviceID() override;
+		virtual void GetControlBusDebugInfo(TArray<FControlBusDebugInfo>& OutDebugInfo) override;
+		virtual void GetControlBusMixDebugInfo(TArray<FControlBusMixDebugInfo>& OutDebugInfo) override;
+		
+
 	private:
 		uint8 bActive : 1;
 		uint8 bShowRenderStatMix : 1;
@@ -101,7 +81,11 @@ namespace AudioModulation
 		uint8 bShowActiveMixes : 1;
 		uint8 bShowGlobalMixes : 1;
 
+		Audio::FDeviceId AssociatedDeviceId;
+
+		TArray<FControlBusDebugInfo> UnfilteredBuses;
 		TArray<FControlBusDebugInfo> FilteredBuses;
+		TArray<FControlBusMixDebugInfo> UnfilteredMixes;
 		TArray<FControlBusMixDebugInfo> FilteredMixes;
 
 		using FGeneratorSortMap = TSortedMap<FString, FGeneratorDebugInfo, FDefaultAllocator, FGeneratorSort>;
diff --git a/Engine/Plugins/Runtime/AudioModulation/Source/AudioModulation/Private/IAudioModulationDebugDataProvider.cpp b/Engine/Plugins/Runtime/AudioModulation/Source/AudioModulation/Private/IAudioModulationDebugDataProvider.cpp
new file mode 100644
--- /dev/null
+++ b/Engine/Plugins/Runtime/AudioModulation/Source/AudioModulation/Private/IAudioModulationDebugDataProvider.cpp
@@ -0,0 +1,13 @@
+﻿// Copyright Epic Games, Inc. All Rights Reserved.
+
+#include "IAudioModulationDebugDataProvider.h"
+
+#if WITH_AUDIOMODULATION
+#if !UE_BUILD_SHIPPING
+
+void AudioModulation::IDebugDataProvider::GetControlBusDebugInfo(TArray<FControlBusDebugInfo>& OutDebugInfo) {}
+void AudioModulation::IDebugDataProvider::GetControlBusMixDebugInfo(TArray<FControlBusMixDebugInfo>& OutDebugInfo) {}
+
+#endif // !UE_BUILD_SHIPPING
+#endif // WITH_AUDIOMODULATION
+
diff --git a/Engine/Plugins/Runtime/AudioModulation/Source/AudioModulation/Public/IAudioModulationDebugDataProvider.h b/Engine/Plugins/Runtime/AudioModulation/Source/AudioModulation/Public/IAudioModulationDebugDataProvider.h
new file mode 100644
--- /dev/null
+++ b/Engine/Plugins/Runtime/AudioModulation/Source/AudioModulation/Public/IAudioModulationDebugDataProvider.h
@@ -0,0 +1,64 @@
+﻿// Copyright Epic Games, Inc. All Rights Reserved.
+
+#pragma once
+
+#include "Misc/Build.h"
+
+#if WITH_AUDIOMODULATION
+#if !UE_BUILD_SHIPPING
+
+#include "AudioDefines.h"
+#include "Containers/Map.h"
+#include "Features/IModularFeature.h"
+#include "UObject/NameTypes.h"
+
+#define UE_API AUDIOMODULATION_API
+
+namespace AudioModulation
+{
+	constexpr FLazyName DebugDataModularFeatureName = "AudioModulationDebugDataProvider";
+	
+	struct FControlBusMixStageDebugInfo
+	{
+		float TargetValue;
+		float CurrentValue;
+	};
+	
+	struct FControlBusMixDebugInfo
+	{
+		FString Name;
+		uint32 Id = -1;
+		uint32 RefCount = 0;
+		bool bIsGlobal = false;
+		double Timer = -1.0;
+
+		TMap<uint32, FControlBusMixStageDebugInfo> Stages;
+	};
+
+	struct FControlBusDebugInfo
+	{
+		FString Name;
+		FString ParameterName;
+		float DefaultValue;
+		float GeneratorValue;
+		float MixValue;
+		float Value;
+		uint32 Id;
+		uint32 RefCount;
+	};
+
+
+	class IDebugDataProvider : public IModularFeature
+	{
+	public:
+		virtual Audio::FDeviceId GetAssociatedDeviceID() = 0;
+		UE_API virtual void GetControlBusDebugInfo(TArray<FControlBusDebugInfo>& OutDebugInfo);
+		UE_API virtual void GetControlBusMixDebugInfo(TArray<FControlBusMixDebugInfo>& OutDebugInfo);
+	};
+}
+
+#undef UE_API
+
+#endif // !UE_BUILD_SHIPPING
+#endif // WITH_AUDIOMODULATION
+
