Fix Control Bus Track crashing in PIE
Don't cache ModulationManager and iterate over all devices when clearing bus mixes.
UE-347686

diff --git a/Engine/Plugins/Runtime/AudioModulation/Source/AudioModulation/Private/MovieSceneAudioControlBusSystem.cpp b/Engine/Plugins/Runtime/AudioModulation/Source/AudioModulation/Private/MovieSceneAudioControlBusSystem.cpp
--- a/Engine/Plugins/Runtime/AudioModulation/Source/AudioModulation/Private/MovieSceneAudioControlBusSystem.cpp
+++ b/Engine/Plugins/Runtime/AudioModulation/Source/AudioModulation/Private/MovieSceneAudioControlBusSystem.cpp
@@ -128,34 +128,26 @@ namespace UE::MovieScene
 			const UWorld* World = PlaybackContext ? PlaybackContext->GetWorld() : nullptr;
 			check(World);
 
-			// If no cached value exists, we need to initialize and put it in pre animated storage so it gets cleared correctly when sequence ends
-			if (!EvaluationData)
+			if (FAudioDevice* AudioDevice = World->GetAudioDeviceRaw())
 			{
-				EvaluationData = AudioSystem->AddAudioControlBusComponentEvaluationData(InstanceHandle, AudioSection, AudioData.Section->MixBus, FName(*MixBusID));
-
-				AudioSystem->PreAnimatedStorage->BeginTrackingEntity(EntityID, true, RootInstanceHandle, AudioData.Section->MixBus);
-				AudioSystem->PreAnimatedStorage->CachePreAnimatedValue(
-					FCachePreAnimatedValueParams(), AudioData.Section->MixBus,
-					[](FObjectKey InKey) { return EPreAnimatedAudioControlBusStateType::ControlBusLifespan; });
-
-				check(AudioSystem->ModulationManager);
-				AudioSystem->ModulationManager->CreateNewSequencerBusMix(World, FName(*MixBusID), {AudioData.Section->MixBus}, BusValue);
-			}
-
-			check(EvaluationData);
-
-			const FAudioDeviceManager* DeviceManager = FAudioDeviceManager::Get();
-			if (ensure(DeviceManager))
-			{
-				const FAudioDevice* AudioDevice = World->GetAudioDevice().GetAudioDevice();
-				if (ensure(AudioDevice))
+				if (AudioDevice->IsModulationPluginEnabled())
 				{
-					if (AudioDevice->IsModulationPluginEnabled())
+					if (IAudioModulationManager* Modulation = AudioDevice->ModulationInterface.Get())
 					{
-						if (IAudioModulationManager* Modulation = AudioDevice->ModulationInterface.Get())
+						// If no cached value exists, we need to initialize and put it in pre animated storage so it gets cleared correctly when sequence ends
+						if (!EvaluationData)
 						{
-							Modulation->UpdateMixByModulator(World, FName(*MixBusID), *AudioData.Section->MixBus, BusValue);
+							EvaluationData = AudioSystem->AddAudioControlBusComponentEvaluationData(InstanceHandle, AudioSection, AudioData.Section->MixBus, FName(*MixBusID));
+							
+							AudioSystem->PreAnimatedStorage->BeginTrackingEntity(EntityID, true, RootInstanceHandle, AudioData.Section->MixBus);
+							AudioSystem->PreAnimatedStorage->CachePreAnimatedValue(
+								FCachePreAnimatedValueParams(), AudioData.Section->MixBus,
+								[](FObjectKey InKey) { return EPreAnimatedAudioControlBusStateType::ControlBusLifespan; });
+							check(AudioDevice->ModulationInterface != nullptr);
+							AudioDevice->ModulationInterface->CreateNewSequencerBusMix(World, FName(*MixBusID), { AudioData.Section->MixBus }, BusValue);
 						}
+						check(EvaluationData);
+						Modulation->UpdateMixByModulator(World, FName(*MixBusID), *AudioData.Section->MixBus, BusValue);
 					}
 				}
 			}
@@ -226,24 +218,6 @@ void UMovieSceneAudioControlBusSystem::OnLink()
 	using namespace UE::MovieScene;
 
 	PreAnimatedStorage = Linker->PreAnimatedState.GetOrCreateStorage<FPreAnimatedAudioControlBusStorage>();
-
-	// Get the modulation manager and create a mix for our sequence
-	FAudioDeviceManager* DeviceManager = FAudioDeviceManager::Get();
-	if (ensure(DeviceManager))
-	{
-		const FAudioDevice* AudioDevice = DeviceManager->GetActiveAudioDevice().GetAudioDevice();
-		if (ensure(AudioDevice))
-		{
-			if (AudioDevice->IsModulationPluginEnabled())
-			{
-				if (IAudioModulationManager* Modulation = AudioDevice->ModulationInterface.Get())
-				{
-					ModulationManager = Modulation;
-					check(ModulationManager);
-				}
-			}
-		}
-	}
 }
 
 void UMovieSceneAudioControlBusSystem::OnUnlink()
@@ -253,16 +227,24 @@ void UMovieSceneAudioControlBusSystem::OnUnlink()
 
 void UMovieSceneAudioControlBusSystem::ClearBusMixes()
 {
-	if (ModulationManager)
+	if (FAudioDeviceManager* DeviceManager = FAudioDeviceManager::Get())
 	{
-		for (const TPair<FObjectKey, AudioControlBusSectionData>& Modulator : ModulatorsByObjectKey)
+		DeviceManager->IterateOverAllDevices([this](::Audio::FDeviceId DeviceId, FAudioDevice* AudioDevice)
 		{
-			ModulationManager->ClearSequencerBusMix(Modulator.Value.ControlBusMixName);
-		}
+			if (AudioDevice->IsModulationPluginEnabled())
+			{
+				if (IAudioModulationManager* ModulationManager = AudioDevice->ModulationInterface.Get())
+				{
+					for (const TPair<FObjectKey, AudioControlBusSectionData>& Modulator : ModulatorsByObjectKey)
+					{
+						ModulationManager->ClearSequencerBusMix(Modulator.Value.ControlBusMixName);
+					}
+				}
+			}
+		});
 	}
 
 	ModulatorsByObjectKey.Reset();
-	ModulationManager = nullptr;
 }
 
 const AudioControlBusSectionData* UMovieSceneAudioControlBusSystem::GetAudioControlBusComponentEvaluationData(const FObjectKey SectionKey) const
diff --git a/Engine/Plugins/Runtime/AudioModulation/Source/AudioModulation/Private/MovieSceneAudioControlBusSystem.h b/Engine/Plugins/Runtime/AudioModulation/Source/AudioModulation/Private/MovieSceneAudioControlBusSystem.h
--- a/Engine/Plugins/Runtime/AudioModulation/Source/AudioModulation/Private/MovieSceneAudioControlBusSystem.h
+++ b/Engine/Plugins/Runtime/AudioModulation/Source/AudioModulation/Private/MovieSceneAudioControlBusSystem.h
@@ -44,7 +44,6 @@ public:
 	const AudioControlBusSectionData* GetAudioControlBusComponentEvaluationData(const FObjectKey SectionKey) const;
 	AudioControlBusSectionData* AddAudioControlBusComponentEvaluationData(FInstanceHandle InstanceHandle, FObjectKey SectionKey, USoundModulatorBase* SectionData, FName ID);
 
-	IAudioModulationManager* ModulationManager = nullptr;
 	TArray<USoundModulatorBase*> Modulators;
 
 	/** Pre-animated state */
