Add equality operator to audio modulation debug data structs
#rnx

diff --git a/Engine/Plugins/Runtime/AudioModulation/Source/AudioModulation/Public/IAudioModulationDebugDataProvider.h b/Engine/Plugins/Runtime/AudioModulation/Source/AudioModulation/Public/IAudioModulationDebugDataProvider.h
--- a/Engine/Plugins/Runtime/AudioModulation/Source/AudioModulation/Public/IAudioModulationDebugDataProvider.h
+++ b/Engine/Plugins/Runtime/AudioModulation/Source/AudioModulation/Public/IAudioModulationDebugDataProvider.h
@@ -30,9 +30,14 @@ namespace AudioModulation
 		uint32 Id = -1;
 		uint32 RefCount = 0;
 		bool bIsGlobal = false;
-		double Timer = -1.0;
+		double Timer = -1.0;		
 
 		TMap<uint32, FControlBusMixStageDebugInfo> Stages;
+
+		bool operator==(const FControlBusMixDebugInfo& Other) const
+		{
+			return this->Id == Other.Id;
+		}
 	};
 
 	struct FControlBusDebugInfo
@@ -45,6 +50,11 @@ namespace AudioModulation
 		float Value;
 		uint32 Id;
 		uint32 RefCount;
+
+		bool operator==(const FControlBusDebugInfo& Other) const
+		{
+			return this->Id == Other.Id;
+		}
 	};
 
 
