Add Type Hash to modulation debug data
#rnx

diff --git a/Engine/Plugins/Runtime/AudioModulation/Source/AudioModulation/Public/IAudioModulationDebugDataProvider.h b/Engine/Plugins/Runtime/AudioModulation/Source/AudioModulation/Public/IAudioModulationDebugDataProvider.h
--- a/Engine/Plugins/Runtime/AudioModulation/Source/AudioModulation/Public/IAudioModulationDebugDataProvider.h
+++ b/Engine/Plugins/Runtime/AudioModulation/Source/AudioModulation/Public/IAudioModulationDebugDataProvider.h
@@ -40,6 +40,11 @@ namespace AudioModulation
 		}
 	};
 
+	FORCEINLINE uint32 GetTypeHash(const FControlBusMixDebugInfo& DebugInfo)
+	{
+		return HashCombine(GetTypeHash(DebugInfo.Name), DebugInfo.Id);
+	}
+
 	struct FControlBusDebugInfo
 	{
 		FString Name;
@@ -57,6 +62,11 @@ namespace AudioModulation
 		}
 	};
 
+	FORCEINLINE uint32 GetTypeHash(const FControlBusDebugInfo& DebugInfo)
+	{
+		return HashCombine(GetTypeHash(DebugInfo.Name), DebugInfo.Id);
+	}
+
 
 	class IDebugDataProvider : public IModularFeature
 	{
