Make operator== for debug data act like type hash

#rnx

diff --git a/Engine/Plugins/Runtime/AudioModulation/Source/AudioModulation/Public/IAudioModulationDebugDataProvider.h b/Engine/Plugins/Runtime/AudioModulation/Source/AudioModulation/Public/IAudioModulationDebugDataProvider.h
--- a/Engine/Plugins/Runtime/AudioModulation/Source/AudioModulation/Public/IAudioModulationDebugDataProvider.h
+++ b/Engine/Plugins/Runtime/AudioModulation/Source/AudioModulation/Public/IAudioModulationDebugDataProvider.h
@@ -36,7 +36,7 @@ namespace AudioModulation
 
 		bool operator==(const FControlBusMixDebugInfo& Other) const
 		{
-			return this->Id == Other.Id;
+			return this->Id == Other.Id && this->Name == Other.Name;
 		}
 	};
 
@@ -58,7 +58,7 @@ namespace AudioModulation
 
 		bool operator==(const FControlBusDebugInfo& Other) const
 		{
-			return this->Id == Other.Id;
+			return this->Id == Other.Id && this->Name == Other.Name;
 		}
 	};
 
