Sequencer: control bus mix track

diff --git a/Engine/Plugins/Runtime/AudioModulation/Source/AudioModulation/Private/MovieSceneAudioControlBusMixSection.cpp b/Engine/Plugins/Runtime/AudioModulation/Source/AudioModulation/Private/MovieSceneAudioControlBusMixSection.cpp
new file mode 100644
--- /dev/null
+++ b/Engine/Plugins/Runtime/AudioModulation/Source/AudioModulation/Private/MovieSceneAudioControlBusMixSection.cpp
@@ -0,0 +1,63 @@
+// Copyright Epic Games, Inc. All Rights Reserved.
+
+#include "MovieSceneAudioControlBusMixSection.h"
+
+#include "Channels/MovieSceneChannelProxy.h"
+#include "EntitySystem/BuiltInComponentTypes.h"
+#include "MovieScene.h"
+#include "MovieSceneAudioModulationComponentTypes.h"
+
+#include UE_INLINE_GENERATED_CPP_BY_NAME(MovieSceneAudioControlBusMixSection)
+
+UMovieSceneAudioControlBusMixSection::UMovieSceneAudioControlBusMixSection(const FObjectInitializer& ObjectInitializer)
+	: Super(ObjectInitializer)
+{
+	BusValue.SetDefault(false);
+	EvalOptions.EnableAndSetCompletionMode(EMovieSceneCompletionMode::ProjectDefault);
+}
+
+EMovieSceneChannelProxyType UMovieSceneAudioControlBusMixSection::CacheChannelProxy()
+{
+	// Set up the channel proxy
+	FMovieSceneChannelProxyData Channels;
+
+#if WITH_EDITOR
+	FMovieSceneChannelMetaData Data;
+	Data.SetIdentifiers(TEXT("Control Bus Mix Value"), FText::FromName(MixBus == nullptr ? TEXT("Control Bus Mix Value") : MixBus->GetFName()));
+
+	Channels.Add(BusValue, Data, TMovieSceneExternalValue<bool>::Make());
+#else
+	Channels.Add(BusValue);
+
+#endif
+
+	ChannelProxy = MakeShared<FMovieSceneChannelProxy>(MoveTemp(Channels));
+
+	return EMovieSceneChannelProxyType::Dynamic;
+}
+
+void UMovieSceneAudioControlBusMixSection::ImportEntityImpl(UMovieSceneEntitySystemLinker* EntityLinker, const FEntityImportParams& Params, FImportedEntity* OutImportedEntity)
+{
+	using namespace UE::MovieScene;
+	const FBuiltInComponentTypes* BuiltInComponents = FBuiltInComponentTypes::Get();
+	const FMovieSceneAudioModulationComponentTypes* TrackComponents = FMovieSceneAudioModulationComponentTypes::Get();
+
+	check(BuiltInComponents);
+	check(TrackComponents);
+
+	const FGuid ObjectBindingID = Params.GetObjectBindingID();
+
+	OutImportedEntity->AddBuilder(
+		FEntityBuilder()
+		.AddConditional(BuiltInComponents->GenericObjectBinding, ObjectBindingID, ObjectBindingID.IsValid())
+		.AddTagConditional(BuiltInComponents->Tags.Root, !ObjectBindingID.IsValid())
+		.Add(TrackComponents->AudioControlBusMix, FMovieSceneAudioControlBusMixComponentData{ this })
+		.Add(BuiltInComponents->BoolChannel, &BusValue)
+	);
+}
+
+void UMovieSceneAudioControlBusMixSection::SetMixBus(USoundControlBusMix* NewMixBus)
+{
+	check(NewMixBus);
+	MixBus = NewMixBus;
+}
diff --git a/Engine/Plugins/Runtime/AudioModulation/Source/AudioModulation/Private/MovieSceneAudioControlBusMixSystem.cpp b/Engine/Plugins/Runtime/AudioModulation/Source/AudioModulation/Private/MovieSceneAudioControlBusMixSystem.cpp
new file mode 100644
--- /dev/null
+++ b/Engine/Plugins/Runtime/AudioModulation/Source/AudioModulation/Private/MovieSceneAudioControlBusMixSystem.cpp
@@ -0,0 +1,267 @@
+// Copyright Epic Games, Inc. All Rights Reserved.
+
+#include "MovieSceneAudioControlBusMixSystem.h"
+
+#include "AudioDevice.h"
+#include "AudioDeviceManager.h"
+#include "Engine/EngineTypes.h"
+#include "EntitySystem/BuiltInComponentTypes.h"
+#include "EntitySystem/IMovieSceneTaskScheduler.h"
+#include "EntitySystem/MovieSceneEntitySystemLinker.h"
+#include "EntitySystem/MovieSceneEntitySystemTask.h"
+#include "EntitySystem/MovieSceneInstanceRegistry.h"
+#include "EntitySystem/MovieSceneEntitySystemRunner.h"
+#include "MovieSceneAudioModulationComponentTypes.h"
+#include "MovieSceneAudioControlBusMixSection.h"
+#include "AudioModulationStatics.h"
+#include "Evaluation/PreAnimatedState/MovieScenePreAnimatedObjectStorage.h"
+#include "Evaluation/PreAnimatedState/MovieScenePreAnimatedStorageID.inl"
+#include "Evaluation/PreAnimatedState/MovieScenePreAnimatedStateExtension.h"
+
+#include UE_INLINE_GENERATED_CPP_BY_NAME(MovieSceneAudioControlBusMixSystem)
+
+namespace UE::MovieScene
+{
+	enum class EPreAnimatedAudioControlBusMixStateType
+	{
+		EnableMixBus,
+		DisableMixBus
+	};
+
+	template<typename BaseTraits>
+	struct FPreAnimatedAudioControlBusMixStateTraits : BaseTraits
+	{
+		using KeyType = FObjectKey;
+		using StorageType = EPreAnimatedAudioControlBusMixStateType;
+
+		EPreAnimatedAudioControlBusMixStateType CachePreAnimatedValue(FObjectKey InKey)
+		{
+			return EPreAnimatedAudioControlBusMixStateType::EnableMixBus;
+		}
+
+		void RestorePreAnimatedValue(FObjectKey InKey, EPreAnimatedAudioControlBusMixStateType InStateType, const FRestoreStateParams& Params)
+		{
+			if (USoundControlBusMix* MixBus = Cast<USoundControlBusMix>(InKey.ResolveObjectPtr()))
+			{
+				UWorld* World = UAudioModulationStatics::GetAudioWorld(MixBus);
+
+				if (World != nullptr)
+				{
+					switch (InStateType)
+					{
+					case UE::MovieScene::EPreAnimatedAudioControlBusMixStateType::EnableMixBus:
+						UAudioModulationStatics::ActivateBusMix(World, MixBus);
+						break;
+					case UE::MovieScene::EPreAnimatedAudioControlBusMixStateType::DisableMixBus:
+						UAudioModulationStatics::DeactivateBusMix(World, MixBus);
+						break;
+					default:
+						break;
+					}
+				}
+			}
+		}
+	};
+
+	using FPreAnimatedBoundObjectAudioControlBusMixStateTraits = FPreAnimatedAudioControlBusMixStateTraits<FBoundObjectPreAnimatedStateTraits>;
+
+	struct FPreAnimatedAudioControlBusMixStorage : TPreAnimatedStateStorage_ObjectTraits<FPreAnimatedBoundObjectAudioControlBusMixStateTraits>
+	{
+		static TAutoRegisterPreAnimatedStorageID<FPreAnimatedAudioControlBusMixStorage> StorageID;
+	};
+	TAutoRegisterPreAnimatedStorageID<FPreAnimatedAudioControlBusMixStorage> FPreAnimatedAudioControlBusMixStorage::StorageID;
+
+	struct FEvaluateAudioControlBusMix
+	{
+		UMovieSceneAudioControlBusMixSystem* AudioSystem;
+		const FInstanceRegistry* InstanceRegistry;
+
+		FEvaluateAudioControlBusMix(UMovieSceneAudioControlBusMixSystem* InAudioSystem)
+			: AudioSystem(InAudioSystem)
+		{
+			check(AudioSystem != nullptr);
+			InstanceRegistry = AudioSystem->GetLinker()->GetInstanceRegistry();
+		}
+
+		void ForEachAllocation(
+			const FEntityAllocation* Allocation,
+			TRead<FMovieSceneEntityID> EntityIDs,
+			TRead<FRootInstanceHandle> RootInstanceHandles,
+			TRead<FInstanceHandle> InstanceHandles,
+			TRead<FMovieSceneAudioControlBusMixComponentData> AudioDatas,
+			TRead<bool> BusValue) const
+		{
+			check(Allocation != nullptr);
+			check(EntityIDs != nullptr);
+			check(RootInstanceHandles != nullptr);
+			check(InstanceHandles != nullptr);
+			check(AudioDatas != nullptr);
+
+			const FBuiltInComponentTypes* BuiltInComponents = FBuiltInComponentTypes::Get();
+
+			const int32 Num = Allocation->Num();
+			const bool bWantsRestoreState = Allocation->HasComponent(BuiltInComponents->Tags.RestoreState);
+
+			for (int32 Index = 0; Index < Num; ++Index)
+			{
+				const FMovieSceneEntityID& EntityID = EntityIDs[Index];
+				const FRootInstanceHandle& RootInstanceHandle = RootInstanceHandles[Index];
+				const FInstanceHandle& InstanceHandle = InstanceHandles[Index];
+
+				const FSequenceInstance& Instance = InstanceRegistry->GetInstance(InstanceHandle);
+
+				bool BusValueData = BusValue[Index];
+
+				Evaluate(EntityID, AudioDatas[Index], Instance, RootInstanceHandle, BusValueData, bWantsRestoreState);
+			}
+		}
+
+	private:
+
+		void Evaluate(
+			const FMovieSceneEntityID& EntityID,
+			const FMovieSceneAudioControlBusMixComponentData& AudioData,
+			const FSequenceInstance& Instance,
+			const FRootInstanceHandle& RootInstanceHandle,
+			bool BusValue,
+			bool bWantsRestoreState) const
+		{
+			const UObject* PlaybackContext = Instance.GetSharedPlaybackState()->GetPlaybackContext();
+			check(PlaybackContext);
+			UWorld* World = PlaybackContext ? PlaybackContext->GetWorld() : nullptr;
+			check(World);
+
+			check(AudioData.Section);
+			USoundControlBusMix* MixBus = AudioData.Section->MixBus;
+
+			check(MixBus);
+
+			AudioModulation::FAudioModulationManager* Manager = UAudioModulationStatics::GetModulation(World);
+			check(Manager);
+
+			UMovieSceneAudioControlBusMixSection* BusMixSection = AudioData.Section;
+			if (!ensureMsgf(BusMixSection, TEXT("No valid control bus mix section found in control bus track component data!")))
+			{
+				return;
+			}
+			FObjectKey SectionKey(BusMixSection);
+			const AudioControlBusMixSectionData* SectionData = AudioSystem->GetSectionData(SectionKey);
+
+			const bool IsBusMixActive = Manager->IsBusMixActive(*AudioData.Section->MixBus);
+
+			if (!SectionData)
+			{
+				AudioSystem->PreAnimatedStorage->BeginTrackingEntity(EntityID, bWantsRestoreState, RootInstanceHandle, MixBus);
+
+				if (IsBusMixActive)
+				{
+					AudioSystem->PreAnimatedStorage->CachePreAnimatedValue(
+						FCachePreAnimatedValueParams(), MixBus,
+						[](FObjectKey InKey) { return EPreAnimatedAudioControlBusMixStateType::EnableMixBus; });
+				}
+				else
+				{
+					AudioSystem->PreAnimatedStorage->CachePreAnimatedValue(
+						FCachePreAnimatedValueParams(), MixBus,
+						[](FObjectKey InKey) { return EPreAnimatedAudioControlBusMixStateType::DisableMixBus; });
+				}
+
+				AudioSystem->AddSectionData(SectionKey, MixBus, IsBusMixActive);
+			}
+
+			if (BusValue && !IsBusMixActive)
+			{
+				UAudioModulationStatics::ActivateBusMix(World, MixBus);
+			}
+			else if (!BusValue && IsBusMixActive)
+			{
+				UAudioModulationStatics::DeactivateBusMix(World, MixBus);
+			}
+		}
+	};
+}
+
+using AudioControlBusMixSectionData = UE::MovieScene::AudioControlBusMixSectionData;
+
+UMovieSceneAudioControlBusMixSystem::UMovieSceneAudioControlBusMixSystem(const FObjectInitializer& ObjInit)
+	: UMovieSceneEntitySystem(ObjInit)
+{
+	using namespace UE::MovieScene;
+
+	RelevantComponent = FMovieSceneAudioModulationComponentTypes::Get()->AudioControlBusMix;
+	Phase = ESystemPhase::Scheduling;
+
+	if (HasAnyFlags(RF_ClassDefaultObject))
+	{
+		const FBuiltInComponentTypes* BuiltInComponents = FBuiltInComponentTypes::Get();
+		check(BuiltInComponents != nullptr);
+
+		DefineComponentConsumer(GetClass(), BuiltInComponents->BoolChannel);
+	}
+}
+
+void UMovieSceneAudioControlBusMixSystem::OnSchedulePersistentTasks(UE::MovieScene::IEntitySystemScheduler* TaskScheduler)
+{
+	using namespace UE::MovieScene;
+
+	const FBuiltInComponentTypes* BuiltInComponents = FBuiltInComponentTypes::Get();
+	const FMovieSceneAudioModulationComponentTypes* TrackComponents = FMovieSceneAudioModulationComponentTypes::Get();
+
+	// Evaluate the control bus mix toggles 
+	FTaskID EvaluateAudioTask = FEntityTaskBuilder()
+		.ReadEntityIDs()
+		.Read(BuiltInComponents->RootInstanceHandle)
+		.Read(BuiltInComponents->InstanceHandle)
+		.Read(TrackComponents->AudioControlBusMix)
+		.Read(BuiltInComponents->BoolResult) // Bus Mix Toggle
+		.SetDesiredThread(Linker->EntityManager.GetGatherThread())
+		.Schedule_PerAllocation<FEvaluateAudioControlBusMix>(&Linker->EntityManager, TaskScheduler, this);
+}
+
+void UMovieSceneAudioControlBusMixSystem::OnRun(FSystemTaskPrerequisites& InPrerequisites, FSystemSubsequentTasks& Subsequents)
+{
+	using namespace UE::MovieScene;
+
+	const FBuiltInComponentTypes* BuiltInComponents = FBuiltInComponentTypes::Get();
+	const FMovieSceneAudioModulationComponentTypes* TrackComponents = FMovieSceneAudioModulationComponentTypes::Get();
+
+	FSystemTaskPrerequisites Prereqs;
+
+	FEntityTaskBuilder()
+		.ReadEntityIDs()
+		.Read(BuiltInComponents->RootInstanceHandle)
+		.Read(BuiltInComponents->InstanceHandle)
+		.Read(TrackComponents->AudioControlBusMix)
+		.Read(BuiltInComponents->BoolResult) // Bus Value
+		.SetDesiredThread(Linker->EntityManager.GetGatherThread())
+		.template Dispatch_PerAllocation<FEvaluateAudioControlBusMix>(&Linker->EntityManager, Prereqs, &Subsequents, this);
+}
+
+void UMovieSceneAudioControlBusMixSystem::OnLink()
+{
+	using namespace UE::MovieScene;
+
+	PreAnimatedStorage = Linker->PreAnimatedState.GetOrCreateStorage<FPreAnimatedAudioControlBusMixStorage>();
+}
+
+void UMovieSceneAudioControlBusMixSystem::OnUnlink()
+{
+	EvaluatedBusMixes.Empty();
+}
+
+const AudioControlBusMixSectionData* UMovieSceneAudioControlBusMixSystem::GetSectionData(FObjectKey Key) const
+{
+	const AudioControlBusMixSectionData* data = EvaluatedBusMixes.Find(Key);
+
+	return data;
+}
+
+AudioControlBusMixSectionData* UMovieSceneAudioControlBusMixSystem::AddSectionData(FObjectKey Key, TObjectPtr<USoundControlBusMix> ControlBusMix, bool bIsActive)
+{
+	using namespace UE::MovieScene;
+	check(ControlBusMix != nullptr);
+
+	AudioControlBusMixSectionData NewData = AudioControlBusMixSectionData(ControlBusMix, bIsActive);
+
+	return &EvaluatedBusMixes.Add(Key, NewData);
+}
diff --git a/Engine/Plugins/Runtime/AudioModulation/Source/AudioModulation/Private/MovieSceneAudioControlBusMixSystem.h b/Engine/Plugins/Runtime/AudioModulation/Source/AudioModulation/Private/MovieSceneAudioControlBusMixSystem.h
new file mode 100644
--- /dev/null
+++ b/Engine/Plugins/Runtime/AudioModulation/Source/AudioModulation/Private/MovieSceneAudioControlBusMixSystem.h
@@ -0,0 +1,51 @@
+// Copyright Epic Games, Inc. All Rights Reserved.
+
+#pragma once
+#include "CoreTypes.h"
+#include "EntitySystem/MovieSceneEntityIDs.h"
+#include "EntitySystem/MovieSceneEntitySystem.h"
+#include "EntitySystem/MovieSceneSequenceInstanceHandle.h"
+#include "MovieSceneTracksComponentTypes.h"
+#include "UObject/ObjectMacros.h"
+#include "IAudioModulation.h"
+#include "AudioModulation.h"
+
+#include "MovieSceneAudioControlBusMixSystem.generated.h"
+
+namespace UE::MovieScene
+{
+	struct FPreAnimatedAudioControlBusMixStorage;
+
+	struct AudioControlBusMixSectionData
+	{
+		TObjectPtr<USoundControlBusMix> ControlBusMix;
+		bool ControlBusMixActive;
+	};
+}
+
+UCLASS()
+class UMovieSceneAudioControlBusMixSystem : public UMovieSceneEntitySystem
+{
+	GENERATED_BODY()
+
+public:
+	using FInstanceHandle = UE::MovieScene::FInstanceHandle;
+	using AudioControlBusMixSectionData = UE::MovieScene::AudioControlBusMixSectionData;
+
+	UMovieSceneAudioControlBusMixSystem(const FObjectInitializer& ObjInit);
+
+	virtual void OnSchedulePersistentTasks(UE::MovieScene::IEntitySystemScheduler* TaskScheduler) override;
+	virtual void OnRun(FSystemTaskPrerequisites& InPrerequisites, FSystemSubsequentTasks& Subsequents) override;
+
+	virtual void OnLink() override;
+	virtual void OnUnlink() override;
+
+	const AudioControlBusMixSectionData* GetSectionData(FObjectKey Key) const;
+	AudioControlBusMixSectionData* AddSectionData(FObjectKey Key, TObjectPtr<USoundControlBusMix> ControlBusMix, bool bIsActive);
+
+	/** Pre-animated state */
+	TSharedPtr<UE::MovieScene::FPreAnimatedAudioControlBusMixStorage> PreAnimatedStorage;
+
+	TMap<FObjectKey, AudioControlBusMixSectionData>  EvaluatedBusMixes;
+};
+
diff --git a/Engine/Plugins/Runtime/AudioModulation/Source/AudioModulation/Private/MovieSceneAudioControlBusMixTrack.cpp b/Engine/Plugins/Runtime/AudioModulation/Source/AudioModulation/Private/MovieSceneAudioControlBusMixTrack.cpp
new file mode 100644
--- /dev/null
+++ b/Engine/Plugins/Runtime/AudioModulation/Source/AudioModulation/Private/MovieSceneAudioControlBusMixTrack.cpp
@@ -0,0 +1,92 @@
+// Copyright Epic Games, Inc. All Rights Reserved.
+
+#include "MovieSceneAudioControlBusMixTrack.h"
+
+#include "Decorations/MovieSceneSectionAnchorsDecoration.h"
+#include "MovieScene.h"
+#include "MovieSceneAudioControlBusMixSection.h"
+
+#include UE_INLINE_GENERATED_CPP_BY_NAME(MovieSceneAudioControlBusMixTrack)
+
+UMovieSceneAudioControlBusMixTrack::UMovieSceneAudioControlBusMixTrack(const FObjectInitializer& ObjectInitializer)
+	: Super(ObjectInitializer)
+{
+}
+
+UMovieSceneSection* UMovieSceneAudioControlBusMixTrack::AddNewControlBusMix(USoundControlBusMix* BusMix, FFrameNumber Time, int32 RowIndex)
+{
+	check(BusMix);
+
+	// add the section
+	UMovieSceneAudioControlBusMixSection* NewSection = Cast<UMovieSceneAudioControlBusMixSection>(CreateNewSection());
+
+	NewSection->SetMixBus(BusMix);
+	FFrameRate FrameRate = GetTypedOuter<UMovieScene>()->GetTickResolution();
+	const float DefaultDuration = 5.0f * FrameRate.AsDecimal();
+	NewSection->InitialPlacementOnRow(MixSections, Time, DefaultDuration, RowIndex);
+	MixSections.Add(NewSection);
+
+#if WITH_EDITORONLY_DATA
+	SetTrackRowDisplayName(FText::FromName(BusMix->GetFName()), NewSection->GetRowIndex());
+#endif
+
+	return NewSection;
+}
+
+bool UMovieSceneAudioControlBusMixTrack::SupportsType(TSubclassOf<UMovieSceneSection> SectionClass) const
+{
+	return SectionClass == UMovieSceneAudioControlBusMixSection::StaticClass();
+}
+
+void UMovieSceneAudioControlBusMixTrack::RemoveAllAnimationData()
+{
+	MixSections.Empty();
+}
+
+bool UMovieSceneAudioControlBusMixTrack::HasSection(const UMovieSceneSection& Section) const
+{
+	return MixSections.Contains(&Section);
+}
+
+void UMovieSceneAudioControlBusMixTrack::AddSection(UMovieSceneSection& Section)
+{
+	MixSections.Add(&Section);
+}
+
+void UMovieSceneAudioControlBusMixTrack::RemoveSection(UMovieSceneSection& Section)
+{
+	MixSections.Remove(&Section);
+}
+
+void UMovieSceneAudioControlBusMixTrack::RemoveSectionAt(int32 SectionIndex)
+{
+	if (SectionIndex >= 0 && SectionIndex < MixSections.Num())
+	{
+		MixSections.RemoveAt(SectionIndex);
+	}
+}
+
+bool UMovieSceneAudioControlBusMixTrack::IsEmpty() const
+{
+	return MixSections.IsEmpty();
+}
+
+const TArray<UMovieSceneSection*>& UMovieSceneAudioControlBusMixTrack::GetAllSections() const
+{
+	return MixSections;
+}
+
+bool UMovieSceneAudioControlBusMixTrack::SupportsMultipleRows() const
+{
+	return true;
+}
+
+UMovieSceneSection* UMovieSceneAudioControlBusMixTrack::CreateNewSection()
+{
+	return NewObject<UMovieSceneAudioControlBusMixSection>(this, NAME_None, RF_Transactional);
+}
+
+void UMovieSceneAudioControlBusMixTrack::GetCompatibleUserDecorationsImpl(TSet<UClass*>& OutClasses) const
+{
+	OutClasses.Add(UMovieSceneSectionAnchorsDecoration::StaticClass());
+}
diff --git a/Engine/Plugins/Runtime/AudioModulation/Source/AudioModulation/Private/MovieSceneAudioControlBusSection.cpp b/Engine/Plugins/Runtime/AudioModulation/Source/AudioModulation/Private/MovieSceneAudioControlBusSection.cpp
--- a/Engine/Plugins/Runtime/AudioModulation/Source/AudioModulation/Private/MovieSceneAudioControlBusSection.cpp
+++ b/Engine/Plugins/Runtime/AudioModulation/Source/AudioModulation/Private/MovieSceneAudioControlBusSection.cpp
@@ -43,6 +43,8 @@ EMovieSceneChannelProxyType UMovieSceneAudioControlBusSection::CacheChannelProxy
 void UMovieSceneAudioControlBusSection::ImportEntityImpl(UMovieSceneEntitySystemLinker* EntityLinker, const FEntityImportParams& Params, FImportedEntity* OutImportedEntity)
 {
 	using namespace UE::MovieScene;
+
+	check(OutImportedEntity);
 	const FBuiltInComponentTypes* BuiltInComponents = FBuiltInComponentTypes::Get();
 	const FMovieSceneAudioModulationComponentTypes* TrackComponents = FMovieSceneAudioModulationComponentTypes::Get();
 
diff --git a/Engine/Plugins/Runtime/AudioModulation/Source/AudioModulation/Private/MovieSceneAudioModulationComponentTypes.cpp b/Engine/Plugins/Runtime/AudioModulation/Source/AudioModulation/Private/MovieSceneAudioModulationComponentTypes.cpp
--- a/Engine/Plugins/Runtime/AudioModulation/Source/AudioModulation/Private/MovieSceneAudioModulationComponentTypes.cpp
+++ b/Engine/Plugins/Runtime/AudioModulation/Source/AudioModulation/Private/MovieSceneAudioModulationComponentTypes.cpp
@@ -17,6 +17,7 @@ FMovieSceneAudioModulationComponentTypes::FMovieSceneAudioModulationComponentTyp
 	FComponentRegistry* ComponentRegistry = UMovieSceneEntitySystemLinker::GetComponents();
 
 	ComponentRegistry->NewComponentType(&AudioControlBus, TEXT("Audio Control Bus"), EComponentTypeFlags::CopyToChildren);
+	ComponentRegistry->NewComponentType(&AudioControlBusMix, TEXT("Audio Control Bus Mix"), EComponentTypeFlags::CopyToChildren);
 }
 
 FMovieSceneAudioModulationComponentTypes::~FMovieSceneAudioModulationComponentTypes()
diff --git a/Engine/Plugins/Runtime/AudioModulation/Source/AudioModulation/Public/MovieSceneAudioControlBusMixSection.h b/Engine/Plugins/Runtime/AudioModulation/Source/AudioModulation/Public/MovieSceneAudioControlBusMixSection.h
new file mode 100644
--- /dev/null
+++ b/Engine/Plugins/Runtime/AudioModulation/Source/AudioModulation/Public/MovieSceneAudioControlBusMixSection.h
@@ -0,0 +1,35 @@
+// Copyright Epic Games, Inc. All Rights Reserved.
+
+#pragma once
+
+#include "EntitySystem/IMovieSceneEntityProvider.h"
+#include "IAudioModulation.h"
+#include "MovieSceneSection.h"
+#include "UObject/ObjectMacros.h"
+#include "Channels/MovieSceneBoolChannel.h"
+#include "SoundControlBusMix.h"
+
+#include "MovieSceneAudioControlBusMixSection.generated.h"
+
+
+UCLASS(MinimalAPI)
+class UMovieSceneAudioControlBusMixSection
+	: public UMovieSceneSection
+	, public IMovieSceneEntityProvider
+{
+	GENERATED_UCLASS_BODY()
+
+public:
+
+	AUDIOMODULATION_API virtual EMovieSceneChannelProxyType CacheChannelProxy() override;
+	AUDIOMODULATION_API virtual void ImportEntityImpl(UMovieSceneEntitySystemLinker* EntityLinker, const FEntityImportParams& Params, FImportedEntity* OutImportedEntity) override;
+
+	AUDIOMODULATION_API void SetMixBus(USoundControlBusMix* NewMixBus);
+
+	UPROPERTY()
+	FMovieSceneBoolChannel BusValue;
+
+	UPROPERTY()
+	TObjectPtr<USoundControlBusMix> MixBus;
+};
+
diff --git a/Engine/Plugins/Runtime/AudioModulation/Source/AudioModulation/Public/MovieSceneAudioControlBusMixTrack.h b/Engine/Plugins/Runtime/AudioModulation/Source/AudioModulation/Public/MovieSceneAudioControlBusMixTrack.h
new file mode 100644
--- /dev/null
+++ b/Engine/Plugins/Runtime/AudioModulation/Source/AudioModulation/Public/MovieSceneAudioControlBusMixTrack.h
@@ -0,0 +1,44 @@
+// Copyright Epic Games, Inc. All Rights Reserved.
+
+#pragma once
+
+#include "IAudioModulation.h"
+#include "MovieSceneNameableTrack.h"
+#include "MovieSceneSection.h"
+#include "SoundControlBusMix.h"
+#include "UObject/ObjectMacros.h"
+
+#include "MovieSceneAudioControlBusMixTrack.generated.h"
+
+
+UCLASS(MinimalAPI)
+class UMovieSceneAudioControlBusMixTrack
+	: public UMovieSceneNameableTrack
+{
+	GENERATED_UCLASS_BODY()
+
+public:
+	AUDIOMODULATION_API UMovieSceneSection* AddNewControlBusMix(USoundControlBusMix* BusMix, FFrameNumber Time, int32 RowIndex);
+
+	// UMovieSceneTrack interface
+
+	AUDIOMODULATION_API virtual bool SupportsType(TSubclassOf<UMovieSceneSection> SectionClass) const override;
+	AUDIOMODULATION_API virtual void RemoveAllAnimationData() override;
+	AUDIOMODULATION_API virtual bool HasSection(const UMovieSceneSection& Section) const override;
+	AUDIOMODULATION_API virtual void AddSection(UMovieSceneSection& Section) override;
+	AUDIOMODULATION_API virtual void RemoveSection(UMovieSceneSection& Section) override;
+	AUDIOMODULATION_API virtual void RemoveSectionAt(int32 SectionIndex) override;
+	AUDIOMODULATION_API virtual bool IsEmpty() const override;
+	AUDIOMODULATION_API virtual const TArray<UMovieSceneSection*>& GetAllSections() const override;
+	AUDIOMODULATION_API virtual bool SupportsMultipleRows() const override;
+	AUDIOMODULATION_API virtual UMovieSceneSection* CreateNewSection() override;
+
+	AUDIOMODULATION_API virtual void GetCompatibleUserDecorationsImpl(TSet<UClass*>& OutClasses) const override;
+
+private:
+
+	/** List of all root mix bus sections */
+	UPROPERTY()
+	TArray<TObjectPtr<UMovieSceneSection>> MixSections;
+};
+
diff --git a/Engine/Plugins/Runtime/AudioModulation/Source/AudioModulation/Public/MovieSceneAudioModulationComponentTypes.h b/Engine/Plugins/Runtime/AudioModulation/Source/AudioModulation/Public/MovieSceneAudioModulationComponentTypes.h
--- a/Engine/Plugins/Runtime/AudioModulation/Source/AudioModulation/Public/MovieSceneAudioModulationComponentTypes.h
+++ b/Engine/Plugins/Runtime/AudioModulation/Source/AudioModulation/Public/MovieSceneAudioModulationComponentTypes.h
@@ -3,11 +3,12 @@
 #pragma once
 
 #include "MovieSceneAudioControlBusSection.h"
+#include "MovieSceneAudioControlBusMixSection.h"
 #include "EntitySystem/MovieSceneEntityIDs.h"
 
 #include "MovieSceneAudioModulationComponentTypes.generated.h"
 
-/** Component data for control bus tracks tracks */
+/** Component data for control bus tracks */
 USTRUCT()
 struct FMovieSceneAudioControlBusComponentData
 {
@@ -17,6 +18,16 @@ struct FMovieSceneAudioControlBusComponentData
 	TObjectPtr<UMovieSceneAudioControlBusSection> Section = nullptr;
 };
 
+/** Component data for control bus mix tracks */
+USTRUCT()
+struct FMovieSceneAudioControlBusMixComponentData
+{
+	GENERATED_BODY()
+
+	UPROPERTY()
+	TObjectPtr<UMovieSceneAudioControlBusMixSection> Section = nullptr;
+};
+
 struct FMovieSceneAudioModulationComponentTypes
 {
 public:
@@ -26,6 +37,7 @@ public:
 	static AUDIOMODULATION_API void Destroy();
 
 	UE::MovieScene::TComponentTypeID<FMovieSceneAudioControlBusComponentData> AudioControlBus;
+	UE::MovieScene::TComponentTypeID<FMovieSceneAudioControlBusMixComponentData> AudioControlBusMix;
 
 private:
 	FMovieSceneAudioModulationComponentTypes();
diff --git a/Engine/Plugins/Runtime/AudioModulation/Source/AudioModulationEditor/Private/AudioControlBusMixTrackEditor.cpp b/Engine/Plugins/Runtime/AudioModulation/Source/AudioModulationEditor/Private/AudioControlBusMixTrackEditor.cpp
new file mode 100644
--- /dev/null
+++ b/Engine/Plugins/Runtime/AudioModulation/Source/AudioModulationEditor/Private/AudioControlBusMixTrackEditor.cpp
@@ -0,0 +1,248 @@
+// Copyright Epic Games, Inc. All Rights Reserved.
+
+#include "AudioControlBusMixTrackEditor.h"
+
+#include "AssetRegistry/AssetRegistryModule.h"
+#include "AudioDevice.h"
+#include "ContentBrowserModule.h"
+#include "Editor.h"
+#include "Engine/World.h"
+#include "IContentBrowserSingleton.h"
+#include "ISequencer.h"
+#include "ISequencerTrackEditor.h"
+#include "MVVM/Extensions/ITrackExtension.h"
+#include "MVVM/ViewModels/OutlinerColumns/OutlinerColumnTypes.h"
+#include "MVVM/Views/ViewUtilities.h"
+#include "SequencerSettings.h"
+#include "MovieSceneAudioControlBusMixTrack.h"
+
+#define LOCTEXT_NAMESPACE "FAudioControlBusMixTrackEditor"
+
+FAudioControlBusMixTrackEditor::FAudioControlBusMixTrackEditor(TSharedRef<ISequencer> InSequencer)
+	: FMovieSceneTrackEditor(InSequencer)
+{
+}
+
+FAudioControlBusMixTrackEditor::~FAudioControlBusMixTrackEditor()
+{
+}
+
+TSharedRef<ISequencerTrackEditor> FAudioControlBusMixTrackEditor::CreateTrackEditor(TSharedRef<ISequencer> OwningSequencer)
+{
+	return MakeShareable(new FAudioControlBusMixTrackEditor(OwningSequencer));
+}
+
+FText FAudioControlBusMixTrackEditor::GetDisplayName() const
+{
+	return LOCTEXT("AudioControlBusMixTrackEditor_DisplayName", "Control Bus Mix");
+}
+
+void FAudioControlBusMixTrackEditor::BuildAddTrackMenu(FMenuBuilder& MenuBuilder)
+{
+	MenuBuilder.AddMenuEntry(
+		LOCTEXT("AddTrack", "Control Bus Mix Track"),
+		LOCTEXT("AddTooltip", "Adds a new control bus mix track that can toggle a mix on or off"),
+		FSlateIcon(FAppStyle::GetAppStyleSetName(), "Sequencer.Tracks.Audio"),
+		FUIAction(
+			FExecuteAction::CreateRaw(this, &FAudioControlBusMixTrackEditor::HandleAddAudioControlBusMixTrackMenuEntryExecute)
+		)
+	);
+}
+
+TSharedPtr<SWidget> FAudioControlBusMixTrackEditor::BuildOutlinerColumnWidget(const FBuildColumnWidgetParams& Params, const FName& ColumnName)
+{
+	using namespace UE::Sequencer;
+
+	UMovieSceneTrack* Track = Params.TrackModel->GetTrack();
+	if (!Track)
+	{
+		return SNullWidget::NullWidget;
+	}
+
+	if (ColumnName == FCommonOutlinerNames::Add)
+	{
+		return UE::Sequencer::MakeAddButton(
+			LOCTEXT("AudioControlBusMixText", "Control Bus Mix"),
+			FOnGetContent::CreateSP(
+				this, &FAudioControlBusMixTrackEditor::BuildAudioControlBusMixSubMenu,
+				FOnAssetSelected::CreateRaw(this, &FAudioControlBusMixTrackEditor::OnAudioControlBusMixAssetSelected, Track),
+				FOnAssetEnterPressed::CreateRaw(this, &FAudioControlBusMixTrackEditor::OnAudioControlBusMixAssetEnterPressed, Track)
+			),
+			Params.ViewModel
+		);
+	}
+
+	return FMovieSceneTrackEditor::BuildOutlinerColumnWidget(Params, ColumnName);
+}
+
+bool FAudioControlBusMixTrackEditor::SupportsType(TSubclassOf<UMovieSceneTrack> Type) const
+{
+	return Type == UMovieSceneAudioControlBusMixTrack::StaticClass();
+}
+
+bool FAudioControlBusMixTrackEditor::SupportsSequence(UMovieSceneSequence* InSequence) const
+{
+	return true;
+}
+
+bool FAudioControlBusMixTrackEditor::IsResizable(UMovieSceneTrack* InTrack) const
+{
+	return true;
+}
+
+void FAudioControlBusMixTrackEditor::HandleAddAudioControlBusMixTrackMenuEntryExecute()
+{
+	if (!IsModulationEnabled())
+	{
+		return;
+	}
+
+	UMovieScene* FocusedMovieScene = GetFocusedMovieScene();
+
+	if (FocusedMovieScene == nullptr)
+	{
+		return;
+	}
+
+	// If movie scene is readonly we cannot add to it, user will have to make it writable and select the add control bus track option again
+	if (FocusedMovieScene->IsReadOnly())
+	{
+		return;
+	}
+
+	const FScopedTransaction Transaction(NSLOCTEXT("Sequencer", "AddAudioControlBusMixTrack_Transaction", "Add Control Bus Mix Track"));
+	FocusedMovieScene->Modify();
+
+	UMovieSceneAudioControlBusMixTrack* NewTrack = FocusedMovieScene->AddTrack<UMovieSceneAudioControlBusMixTrack>();
+	ensure(NewTrack);
+
+	if (NewTrack == nullptr)
+	{
+		return;
+	}
+
+	NewTrack->SetDisplayName(LOCTEXT("AudioControlBusMixTrackName", "Control Bus Mix"));
+
+	if (GetSequencer().IsValid())
+	{
+		GetSequencer()->OnAddTrack(NewTrack, FGuid());
+	}
+}
+
+void FAudioControlBusMixTrackEditor::OnAudioControlBusMixAssetSelected(const FAssetData& AssetData, UMovieSceneTrack* Track)
+{
+	FSlateApplication::Get().DismissAllMenus();
+
+	if (!IsModulationEnabled())
+	{
+		return;
+	}
+
+	UObject* SelectedObject = AssetData.GetAsset();
+
+	if (SelectedObject)
+	{
+		USoundControlBusMix* NewBusMix = CastChecked<USoundControlBusMix>(AssetData.GetAsset());
+		if (NewBusMix != nullptr)
+		{
+			const FScopedTransaction Transaction(NSLOCTEXT("Sequencer", "AddAudioControlBusMix_Transaction", "Add Control Bus Mix"));
+
+			UMovieSceneAudioControlBusMixTrack* AudioControlBusMixTrack = Cast<UMovieSceneAudioControlBusMixTrack>(Track);
+			check(AudioControlBusMixTrack);
+
+			AudioControlBusMixTrack->Modify();
+
+			check(GetSequencer().IsValid());
+
+			FFrameTime KeyTime = GetSequencer()->GetLocalTime().Time;
+
+			UMovieSceneSection* NewSection = AudioControlBusMixTrack->AddNewControlBusMix(NewBusMix, KeyTime.FrameNumber, INDEX_NONE);
+
+			GetSequencer()->EmptySelection();
+			GetSequencer()->SelectSection(NewSection);
+			GetSequencer()->ThrobSectionSelection();
+
+			GetSequencer()->NotifyMovieSceneDataChanged(EMovieSceneDataChangeType::MovieSceneStructureItemAdded);
+		}
+	}
+}
+
+void FAudioControlBusMixTrackEditor::OnAudioControlBusMixAssetEnterPressed(const TArray<FAssetData>& AssetData, UMovieSceneTrack* Track)
+{
+	// By default enter will select the first list element if there is at least 1 element
+	if (AssetData.Num() > 0)
+	{
+		OnAudioControlBusMixAssetSelected(AssetData[0].GetAsset(), Track);
+	}
+}
+
+TSharedRef<SWidget> FAudioControlBusMixTrackEditor::BuildAudioControlBusMixSubMenu(FOnAssetSelected OnAssetSelected, FOnAssetEnterPressed OnAssetEnterPressed)
+{
+	FMenuBuilder MenuBuilder(true, nullptr);
+
+	if (!IsModulationEnabled())
+	{
+		return MenuBuilder.MakeWidget();
+	}
+
+	const TSharedPtr<ISequencer> SequencerPtr = GetSequencer();
+	const UMovieSceneSequence* Sequence = SequencerPtr.IsValid() ? SequencerPtr->GetFocusedMovieSceneSequence() : nullptr;
+
+	FAssetRegistryModule& AssetRegistryModule = FModuleManager::LoadModuleChecked<FAssetRegistryModule>(TEXT("AssetRegistry"));
+	TArray<FTopLevelAssetPath> ClassNames;
+	ClassNames.Add(USoundControlBusMix::StaticClass()->GetClassPathName());
+
+	TSet<FTopLevelAssetPath> DerivedClassNames;
+	AssetRegistryModule.Get().GetDerivedClassNames(ClassNames, TSet<FTopLevelAssetPath>(), DerivedClassNames);
+
+	FAssetPickerConfig AssetPickerConfig;
+	{
+		AssetPickerConfig.OnAssetSelected = OnAssetSelected;
+		AssetPickerConfig.OnAssetEnterPressed = OnAssetEnterPressed;
+		AssetPickerConfig.bAllowNullSelection = false;
+		AssetPickerConfig.bAddFilterUI = true;
+		AssetPickerConfig.bShowTypeInColumnView = false;
+		AssetPickerConfig.InitialAssetViewType = EAssetViewType::List;
+		for (FTopLevelAssetPath ClassName : DerivedClassNames)
+		{
+			AssetPickerConfig.Filter.ClassPaths.Add(ClassName);
+		}
+		AssetPickerConfig.SaveSettingsName = TEXT("SequencerAssetPicker");
+		AssetPickerConfig.AdditionalReferencingAssets.Add(FAssetData(Sequence));
+	}
+
+	FContentBrowserModule& ContentBrowserModule = FModuleManager::Get().LoadModuleChecked<FContentBrowserModule>(TEXT("ContentBrowser"));
+
+	const float WidthOverride = SequencerPtr.IsValid() ? SequencerPtr->GetSequencerSettings()->GetAssetBrowserWidth() : 500.f;
+	const float HeightOverride = SequencerPtr.IsValid() ? SequencerPtr->GetSequencerSettings()->GetAssetBrowserHeight() : 400.f;
+
+	TSharedPtr<SBox> MenuEntry = SNew(SBox)
+		.WidthOverride(WidthOverride)
+		.HeightOverride(HeightOverride)
+		[
+			ContentBrowserModule.Get().CreateAssetPicker(AssetPickerConfig)
+		];
+
+	MenuBuilder.AddWidget(MenuEntry.ToSharedRef(), FText::GetEmpty(), true);
+
+	return MenuBuilder.MakeWidget();
+}
+
+bool FAudioControlBusMixTrackEditor::IsModulationEnabled()
+{
+	if (GEditor)
+	{
+		if (UWorld* World = GEditor->GetEditorWorldContext().World())
+		{
+			FAudioDeviceHandle AudioDeviceHandle = World->GetAudioDevice();
+			if (AudioDeviceHandle.IsValid())
+			{
+				return AudioDeviceHandle->IsModulationPluginEnabled();
+			}
+		}
+	}
+
+	return false;
+}
+
+#undef LOCTEXT_NAMESPACE
\ No newline at end of file
diff --git a/Engine/Plugins/Runtime/AudioModulation/Source/AudioModulationEditor/Private/AudioControlBusMixTrackEditor.h b/Engine/Plugins/Runtime/AudioModulation/Source/AudioModulationEditor/Private/AudioControlBusMixTrackEditor.h
new file mode 100644
--- /dev/null
+++ b/Engine/Plugins/Runtime/AudioModulation/Source/AudioModulationEditor/Private/AudioControlBusMixTrackEditor.h
@@ -0,0 +1,61 @@
+// Copyright Epic Games, Inc. All Rights Reserved.
+
+#pragma once
+
+#include "IAudioModulation.h"
+#include "IContentBrowserSingleton.h"
+#include "ISequencer.h"
+#include "ISequencerTrackEditor.h"
+#include "MovieSceneSection.h"
+#include "MovieSceneTrackEditor.h"
+
+#define UE_API AUDIOMODULATIONEDITOR_API
+
+class FAudioControlBusMixTrackEditor
+	: public FMovieSceneTrackEditor
+{
+public:
+	/**
+	* Constructor
+	*
+	* @param InSequencer The sequencer instance to be used by this tool
+	*/
+	UE_API FAudioControlBusMixTrackEditor(TSharedRef<ISequencer> InSequencer);
+
+	/** Virtual destructor. */
+	UE_API virtual ~FAudioControlBusMixTrackEditor();
+
+	/**
+	 * Creates an instance of this class.  Called by a sequencer
+	 *
+	 * @param OwningSequencer The sequencer instance to be used by this tool
+	 * @return The new instance of this class
+	 */
+	static UE_API TSharedRef<ISequencerTrackEditor> CreateTrackEditor(TSharedRef<ISequencer> OwningSequencer);
+
+	// ISequencerTrackEditor interface
+
+	UE_API virtual FText GetDisplayName() const override;
+	UE_API virtual void BuildAddTrackMenu(FMenuBuilder& MenuBuilder) override;
+	UE_API virtual TSharedPtr<SWidget> BuildOutlinerColumnWidget(const FBuildColumnWidgetParams& Params, const FName& ColumnName) override;
+	UE_API virtual bool SupportsType(TSubclassOf<UMovieSceneTrack> Type) const override;
+	UE_API virtual bool SupportsSequence(UMovieSceneSequence* InSequence) const override;
+	UE_API virtual bool IsResizable(UMovieSceneTrack* InTrack) const override;
+
+private:
+	/** Callback for executing the "Add Audio Control Bus Mix Track" menu entry. */
+	UE_API void HandleAddAudioControlBusMixTrackMenuEntryExecute();
+
+	/** Audio control bus mix asset selected */
+	UE_API void OnAudioControlBusMixAssetSelected(const FAssetData& AssetData, UMovieSceneTrack* Track);
+
+	/** Audio control bus mix asset enter pressed */
+	UE_API void OnAudioControlBusMixAssetEnterPressed(const TArray<FAssetData>& AssetData, UMovieSceneTrack* Track);
+
+	/** Audio control bus mix sub menu */
+	UE_API TSharedRef<SWidget> BuildAudioControlBusMixSubMenu(FOnAssetSelected OnAssetSelected, FOnAssetEnterPressed OnAssetEnterPressed);
+
+	UE_API bool IsModulationEnabled();
+};
+
+#undef UE_API
\ No newline at end of file
diff --git a/Engine/Plugins/Runtime/AudioModulation/Source/AudioModulationEditor/Private/AudioModulationEditor.cpp b/Engine/Plugins/Runtime/AudioModulation/Source/AudioModulationEditor/Private/AudioModulationEditor.cpp
--- a/Engine/Plugins/Runtime/AudioModulation/Source/AudioModulationEditor/Private/AudioModulationEditor.cpp
+++ b/Engine/Plugins/Runtime/AudioModulation/Source/AudioModulationEditor/Private/AudioModulationEditor.cpp
@@ -31,6 +31,7 @@
 #include "UObject/UObjectIterator.h"
 #include "ISequencerModule.h"
 #include "AudioControlBusTrackEditor.h"
+#include "AudioControlBusMixTrackEditor.h"
 
 #if WITH_AUDIOMODULATION_METASOUND_SUPPORT
 #include "MetasoundDataReference.h"
@@ -126,6 +127,7 @@ void FAudioModulationEditorModule::StartupModule()
 	{
 		ISequencerModule& SequencerModule = FModuleManager::Get().LoadModuleChecked<ISequencerModule>("Sequencer");
 		AudioControlBusTrackCreateEditorHandle = SequencerModule.RegisterTrackEditor(FOnCreateTrackEditor::CreateStatic(&FAudioControlBusTrackEditor::CreateTrackEditor));
+		AudioControlBusMixTrackCreateEditorHandle = SequencerModule.RegisterTrackEditor(FOnCreateTrackEditor::CreateStatic(&FAudioControlBusMixTrackEditor::CreateTrackEditor));
 	}
 #endif
 
@@ -212,6 +214,7 @@ void FAudioModulationEditorModule::ShutdownModule()
 	{
 		ISequencerModule& SequencerModule = FModuleManager::Get().LoadModuleChecked<ISequencerModule>("Sequencer");
 		SequencerModule.UnRegisterTrackEditor(AudioControlBusTrackCreateEditorHandle);
+		SequencerModule.UnRegisterTrackEditor(AudioControlBusMixTrackCreateEditorHandle);
 	}
 #endif
 }
diff --git a/Engine/Plugins/Runtime/AudioModulation/Source/AudioModulationEditor/Private/AudioModulationEditor.h b/Engine/Plugins/Runtime/AudioModulation/Source/AudioModulationEditor/Private/AudioModulationEditor.h
--- a/Engine/Plugins/Runtime/AudioModulation/Source/AudioModulationEditor/Private/AudioModulationEditor.h
+++ b/Engine/Plugins/Runtime/AudioModulation/Source/AudioModulationEditor/Private/AudioModulationEditor.h
@@ -37,4 +37,5 @@ private:
 	TSharedPtr<FSlateStyleSet> StyleSet;
 
 	FDelegateHandle AudioControlBusTrackCreateEditorHandle;
+	FDelegateHandle AudioControlBusMixTrackCreateEditorHandle;
 };
