Sequencer: Control bus tracks UX refactor

diff --git a/Engine/Plugins/Runtime/AudioModulation/Source/AudioModulation/Private/MovieSceneAudioControlBusMixTrack.cpp b/Engine/Plugins/Runtime/AudioModulation/Source/AudioModulation/Private/MovieSceneAudioControlBusMixTrack.cpp
--- a/Engine/Plugins/Runtime/AudioModulation/Source/AudioModulation/Private/MovieSceneAudioControlBusMixTrack.cpp
+++ b/Engine/Plugins/Runtime/AudioModulation/Source/AudioModulation/Private/MovieSceneAudioControlBusMixTrack.cpp
@@ -3,8 +3,8 @@
 #include "MovieSceneAudioControlBusMixTrack.h"
 
 #include "Decorations/MovieSceneSectionAnchorsDecoration.h"
-#include "MovieScene.h"
 #include "MovieSceneAudioControlBusMixSection.h"
+#include "MovieSceneSequence.h"
 
 #include UE_INLINE_GENERATED_CPP_BY_NAME(MovieSceneAudioControlBusMixTrack)
 
@@ -13,7 +13,7 @@ UMovieSceneAudioControlBusMixTrack::UMovieSceneAudioControlBusMixTrack(const FOb
 {
 }
 
-UMovieSceneSection* UMovieSceneAudioControlBusMixTrack::AddNewControlBusMix(USoundControlBusMix* BusMix, FFrameNumber Time, int32 RowIndex)
+UMovieSceneSection* UMovieSceneAudioControlBusMixTrack::AddNewControlBusMix(USoundControlBusMix* BusMix)
 {
 	check(BusMix);
 
@@ -21,9 +21,8 @@ UMovieSceneSection* UMovieSceneAudioControlBusMixTrack::AddNewControlBusMix(USou
 	UMovieSceneAudioControlBusMixSection* NewSection = Cast<UMovieSceneAudioControlBusMixSection>(CreateNewSection());
 
 	NewSection->SetMixBus(BusMix);
-	FFrameRate FrameRate = GetTypedOuter<UMovieScene>()->GetTickResolution();
-	const float DefaultDuration = 5.0f * FrameRate.AsDecimal();
-	NewSection->InitialPlacementOnRow(MixSections, Time, DefaultDuration, RowIndex);
+
+	NewSection->SetRange(TRange<FFrameNumber>::All());
 	MixSections.Add(NewSection);
 
 #if WITH_EDITORONLY_DATA
@@ -90,3 +89,27 @@ void UMovieSceneAudioControlBusMixTrack::GetCompatibleUserDecorationsImpl(TSet<U
 {
 	OutClasses.Add(UMovieSceneSectionAnchorsDecoration::StaticClass());
 }
+
+#if WITH_EDITOR
+FText UMovieSceneAudioControlBusMixTrack::GetDefaultDisplayName() const
+{
+	if (MixSections.Num() > 0)
+	{
+		return FText::FromString(Cast<UMovieSceneAudioControlBusMixSection>(MixSections[0])->MixBus->GetName());
+	}
+
+	return NSLOCTEXT("MovieSceneAudioControlBusMixTrack", "DefaultTrackName", "Control Bus Mix");
+}
+
+void UMovieSceneAudioControlBusMixTrack::PostLoad()
+{
+	if (MixSections.Num() > 0 && GetDisplayName().EqualTo(FText::FromString(Cast<UMovieSceneAudioControlBusMixSection>(MixSections[0])->MixBus->GetName())))
+	{
+		UMovieSceneSequence* Sequence = GetTypedOuter<UMovieSceneSequence>();
+		UE_LOG(LogMovieScene, Display, TEXT("Control Bus Mix Track: Removing display name %s from %s because it matches default asset name and will not update automatically with asset changes"), *MixSections[0]->GetName(), *GetNameSafe(Sequence));
+		SetDisplayName(FText::GetEmpty());
+	}
+	Super::PostLoad();
+}
+
+#endif
diff --git a/Engine/Plugins/Runtime/AudioModulation/Source/AudioModulation/Private/MovieSceneAudioControlBusTrack.cpp b/Engine/Plugins/Runtime/AudioModulation/Source/AudioModulation/Private/MovieSceneAudioControlBusTrack.cpp
--- a/Engine/Plugins/Runtime/AudioModulation/Source/AudioModulation/Private/MovieSceneAudioControlBusTrack.cpp
+++ b/Engine/Plugins/Runtime/AudioModulation/Source/AudioModulation/Private/MovieSceneAudioControlBusTrack.cpp
@@ -3,8 +3,8 @@
 #include "MovieSceneAudioControlBusTrack.h"
 
 #include "Decorations/MovieSceneSectionAnchorsDecoration.h"
-#include "MovieScene.h"
 #include "MovieSceneAudioControlBusSection.h"
+#include "MovieSceneSequence.h"
 
 #include UE_INLINE_GENERATED_CPP_BY_NAME(MovieSceneAudioControlBusTrack)
 
@@ -13,7 +13,7 @@ UMovieSceneAudioControlBusTrack::UMovieSceneAudioControlBusTrack(const FObjectIn
 {
 }
 
-UMovieSceneSection* UMovieSceneAudioControlBusTrack::AddNewControlBus(USoundModulatorBase* Bus, FFrameNumber Time, int32 RowIndex)
+UMovieSceneSection* UMovieSceneAudioControlBusTrack::AddNewControlBus(USoundModulatorBase* Bus)
 {
 	check(Bus);
 
@@ -21,9 +21,8 @@ UMovieSceneSection* UMovieSceneAudioControlBusTrack::AddNewControlBus(USoundModu
 	UMovieSceneAudioControlBusSection* NewSection = Cast<UMovieSceneAudioControlBusSection>(CreateNewSection());
 
 	NewSection->SetBus(Bus);
-	FFrameRate FrameRate = GetTypedOuter<UMovieScene>()->GetTickResolution();
-	const float DefaultDuration = 5.0f * FrameRate.AsDecimal();
-	NewSection->InitialPlacementOnRow(AudioSections, Time, DefaultDuration, RowIndex);
+
+	NewSection->SetRange(TRange<FFrameNumber>::All());
 	AudioSections.Add(NewSection);
 
 #if WITH_EDITORONLY_DATA
@@ -87,3 +86,27 @@ void UMovieSceneAudioControlBusTrack::GetCompatibleUserDecorationsImpl(TSet<UCla
 {
 	OutClasses.Add(UMovieSceneSectionAnchorsDecoration::StaticClass());
 }
+
+#if WITH_EDITOR
+FText UMovieSceneAudioControlBusTrack::GetDefaultDisplayName() const
+{
+	if (AudioSections.Num() > 0)
+	{
+		return FText::FromString(Cast<UMovieSceneAudioControlBusSection>(AudioSections[0])->MixBus->GetName());
+	}
+
+	return NSLOCTEXT("MovieSceneAudioControlBusTrack","DefaultTrackName", "Control Bus");
+}
+
+void UMovieSceneAudioControlBusTrack::PostLoad()
+{
+	if (AudioSections.Num() > 0 && GetDisplayName().EqualTo(FText::FromString(Cast<UMovieSceneAudioControlBusSection>(AudioSections[0])->MixBus->GetName())))
+	{
+		UMovieSceneSequence* Sequence = GetTypedOuter<UMovieSceneSequence>();
+		UE_LOG(LogMovieScene, Display, TEXT("Control Bus Track: Removing display name %s from %s because it matches default asset name and will not update automatically with asset changes"), *AudioSections[0]->GetName(), *GetNameSafe(Sequence));
+		SetDisplayName(FText::GetEmpty());
+	}
+	Super::PostLoad();
+}
+
+#endif
diff --git a/Engine/Plugins/Runtime/AudioModulation/Source/AudioModulation/Public/MovieSceneAudioControlBusMixTrack.h b/Engine/Plugins/Runtime/AudioModulation/Source/AudioModulation/Public/MovieSceneAudioControlBusMixTrack.h
--- a/Engine/Plugins/Runtime/AudioModulation/Source/AudioModulation/Public/MovieSceneAudioControlBusMixTrack.h
+++ b/Engine/Plugins/Runtime/AudioModulation/Source/AudioModulation/Public/MovieSceneAudioControlBusMixTrack.h
@@ -2,7 +2,6 @@
 
 #pragma once
 
-#include "IAudioModulation.h"
 #include "MovieSceneNameableTrack.h"
 #include "MovieSceneSection.h"
 #include "SoundControlBusMix.h"
@@ -18,7 +17,7 @@ class UMovieSceneAudioControlBusMixTrack
 	GENERATED_UCLASS_BODY()
 
 public:
-	AUDIOMODULATION_API UMovieSceneSection* AddNewControlBusMix(USoundControlBusMix* BusMix, FFrameNumber Time, int32 RowIndex);
+	AUDIOMODULATION_API UMovieSceneSection* AddNewControlBusMix(USoundControlBusMix* BusMix);
 
 	// UMovieSceneTrack interface
 
@@ -35,8 +34,17 @@ public:
 
 	AUDIOMODULATION_API virtual void GetCompatibleUserDecorationsImpl(TSet<UClass*>& OutClasses) const override;
 
+#if WITH_EDITOR
+	AUDIOMODULATION_API virtual FText GetDefaultDisplayName() const override;
+#endif
+
 private:
 
+	//~ UObject interface
+#if WITH_EDITOR
+	AUDIOMODULATION_API virtual void PostLoad() override;
+#endif
+
 	/** List of all root mix bus sections */
 	UPROPERTY()
 	TArray<TObjectPtr<UMovieSceneSection>> MixSections;
diff --git a/Engine/Plugins/Runtime/AudioModulation/Source/AudioModulation/Public/MovieSceneAudioControlBusTrack.h b/Engine/Plugins/Runtime/AudioModulation/Source/AudioModulation/Public/MovieSceneAudioControlBusTrack.h
--- a/Engine/Plugins/Runtime/AudioModulation/Source/AudioModulation/Public/MovieSceneAudioControlBusTrack.h
+++ b/Engine/Plugins/Runtime/AudioModulation/Source/AudioModulation/Public/MovieSceneAudioControlBusTrack.h
@@ -17,7 +17,7 @@ class UMovieSceneAudioControlBusTrack
 
 public:
 
-	AUDIOMODULATION_API UMovieSceneSection* AddNewControlBus(USoundModulatorBase* Bus, FFrameNumber Time, int32 RowIndex);
+	AUDIOMODULATION_API UMovieSceneSection* AddNewControlBus(USoundModulatorBase* Bus);
 
 	// UMovieSceneTrack interface
 
@@ -34,8 +34,17 @@ public:
 
 	AUDIOMODULATION_API virtual void GetCompatibleUserDecorationsImpl(TSet<UClass*>& OutClasses) const override;
 
+#if WITH_EDITOR
+	AUDIOMODULATION_API virtual FText GetDefaultDisplayName() const override;
+#endif
+
 private:
 
+	//~ UObject interface
+#if WITH_EDITOR
+	AUDIOMODULATION_API virtual void PostLoad() override;
+#endif
+
 	/** List of all root audio sections */
 	UPROPERTY()
 	TArray<TObjectPtr<UMovieSceneSection>> AudioSections;
diff --git a/Engine/Plugins/Runtime/AudioModulation/Source/AudioModulationEditor/Private/AudioControlBusBaseTrackEditor.cpp b/Engine/Plugins/Runtime/AudioModulation/Source/AudioModulationEditor/Private/AudioControlBusBaseTrackEditor.cpp
new file mode 100644
--- /dev/null
+++ b/Engine/Plugins/Runtime/AudioModulation/Source/AudioModulationEditor/Private/AudioControlBusBaseTrackEditor.cpp
@@ -0,0 +1,37 @@
+// Copyright Epic Games, Inc. All Rights Reserved.
+
+#include "AudioControlBusBaseTrackEditor.h"
+
+#include "AssetRegistry/AssetData.h"
+
+FAudioControlBusBaseTrackEditor::FAudioControlBusBaseTrackEditor(TSharedRef<ISequencer> InSequencer)
+	: FMovieSceneTrackEditor(InSequencer)
+{
+}
+
+FAudioControlBusBaseTrackEditor::~FAudioControlBusBaseTrackEditor()
+{
+}
+
+bool FAudioControlBusBaseTrackEditor::SupportsSequence(UMovieSceneSequence* InSequence) const
+{
+	return true;
+}
+
+bool FAudioControlBusBaseTrackEditor::IsResizable(UMovieSceneTrack* InTrack) const
+{
+	return false;
+}
+
+void FAudioControlBusBaseTrackEditor::AddControlBusAssetTrackToSequence(const FAssetData& InAssetData)
+{
+}
+
+void FAudioControlBusBaseTrackEditor::OnAudioControlBusAssetEnterPressed(const TArray<FAssetData>& AssetData)
+{
+	// By default enter will select the first list element if there is at least 1 element
+	if (AssetData.Num() > 0)
+	{
+		AddControlBusAssetTrackToSequence(AssetData[0]);
+	}
+}
diff --git a/Engine/Plugins/Runtime/AudioModulation/Source/AudioModulationEditor/Private/AudioControlBusBaseTrackEditor.h b/Engine/Plugins/Runtime/AudioModulation/Source/AudioModulationEditor/Private/AudioControlBusBaseTrackEditor.h
new file mode 100644
--- /dev/null
+++ b/Engine/Plugins/Runtime/AudioModulation/Source/AudioModulationEditor/Private/AudioControlBusBaseTrackEditor.h
@@ -0,0 +1,36 @@
+// Copyright Epic Games, Inc. All Rights Reserved.
+
+#pragma once
+
+
+#include "ISequencer.h"
+#include "MovieSceneTrackEditor.h"
+
+#define UE_API AUDIOMODULATIONEDITOR_API
+
+class FAudioControlBusBaseTrackEditor
+	: public FMovieSceneTrackEditor
+{
+public:
+	/**
+	* Constructor
+	*
+	* @param InSequencer The sequencer instance to be used by this tool
+	*/
+	UE_API FAudioControlBusBaseTrackEditor(TSharedRef<ISequencer> InSequencer);
+
+	/** Virtual destructor. */
+	UE_API virtual ~FAudioControlBusBaseTrackEditor();
+
+	UE_API virtual bool SupportsSequence(UMovieSceneSequence* InSequence) const override;
+	UE_API virtual bool IsResizable(UMovieSceneTrack* InTrack) const override;
+	/** Audio control bus asset enter pressed */
+	UE_API void OnAudioControlBusAssetEnterPressed(const TArray<FAssetData>& AssetData);
+
+protected:
+	/** Add new control bus Asset track to sequence */
+	UE_API virtual void AddControlBusAssetTrackToSequence(const FAssetData& InAssetData);
+};
+
+#undef UE_API
+
diff --git a/Engine/Plugins/Runtime/AudioModulation/Source/AudioModulationEditor/Private/AudioControlBusMixTrackEditor.cpp b/Engine/Plugins/Runtime/AudioModulation/Source/AudioModulationEditor/Private/AudioControlBusMixTrackEditor.cpp
--- a/Engine/Plugins/Runtime/AudioModulation/Source/AudioModulationEditor/Private/AudioControlBusMixTrackEditor.cpp
+++ b/Engine/Plugins/Runtime/AudioModulation/Source/AudioModulationEditor/Private/AudioControlBusMixTrackEditor.cpp
@@ -2,24 +2,18 @@
 
 #include "AudioControlBusMixTrackEditor.h"
 
-#include "AssetRegistry/AssetRegistryModule.h"
-#include "AudioDevice.h"
 #include "ContentBrowserModule.h"
-#include "Editor.h"
-#include "Engine/World.h"
 #include "IContentBrowserSingleton.h"
 #include "ISequencer.h"
 #include "ISequencerTrackEditor.h"
-#include "MVVM/Extensions/ITrackExtension.h"
-#include "MVVM/ViewModels/OutlinerColumns/OutlinerColumnTypes.h"
-#include "MVVM/Views/ViewUtilities.h"
-#include "SequencerSettings.h"
 #include "MovieSceneAudioControlBusMixTrack.h"
+#include "SequencerSettings.h"
+#include "Styling/SlateIconFinder.h"
 
 #define LOCTEXT_NAMESPACE "FAudioControlBusMixTrackEditor"
 
 FAudioControlBusMixTrackEditor::FAudioControlBusMixTrackEditor(TSharedRef<ISequencer> InSequencer)
-	: FMovieSceneTrackEditor(InSequencer)
+	: FAudioControlBusBaseTrackEditor(InSequencer)
 {
 }
 
@@ -39,89 +33,82 @@ FText FAudioControlBusMixTrackEditor::GetDisplayName() const
 
 void FAudioControlBusMixTrackEditor::BuildAddTrackMenu(FMenuBuilder& MenuBuilder)
 {
-	MenuBuilder.AddMenuEntry(
-		LOCTEXT("AddTrack", "Control Bus Mix Track"),
-		LOCTEXT("AddTooltip", "Adds a new control bus mix track that can toggle a mix on or off"),
-		FSlateIcon(FAppStyle::GetAppStyleSetName(), "Sequencer.Tracks.Audio"),
-		FUIAction(
-			FExecuteAction::CreateRaw(this, &FAudioControlBusMixTrackEditor::HandleAddAudioControlBusMixTrackMenuEntryExecute)
-		)
-	);
-}
+	auto SubMenuCallback = [this](FMenuBuilder& SubMenuBuilder)
+		{
+			TSharedPtr<ISequencer> Sequencer = GetSequencer();
+			UMovieSceneSequence* Sequence = Sequencer.IsValid() ? Sequencer->GetFocusedMovieSceneSequence() : nullptr;
 
-TSharedPtr<SWidget> FAudioControlBusMixTrackEditor::BuildOutlinerColumnWidget(const FBuildColumnWidgetParams& Params, const FName& ColumnName)
-{
-	using namespace UE::Sequencer;
+			FAssetPickerConfig AssetPickerConfig;
+			{
+				AssetPickerConfig.OnAssetSelected = FOnAssetSelected::CreateRaw(this, &FAudioControlBusMixTrackEditor::AddControlBusAssetTrackToSequence);
+				AssetPickerConfig.OnAssetEnterPressed = FOnAssetEnterPressed::CreateRaw(this, &FAudioControlBusBaseTrackEditor::OnAudioControlBusAssetEnterPressed);
+				AssetPickerConfig.bAllowNullSelection = false;
+				AssetPickerConfig.bAddFilterUI = true;
+				AssetPickerConfig.InitialAssetViewType = EAssetViewType::List;
+				AssetPickerConfig.Filter.bRecursiveClasses = true;
+				AssetPickerConfig.Filter.ClassPaths.Add(USoundControlBusMix::StaticClass()->GetClassPathName());
+				AssetPickerConfig.SaveSettingsName = TEXT("SequencerAssetPicker");
+				AssetPickerConfig.AdditionalReferencingAssets.Add(FAssetData(Sequence));
+			}
 
-	UMovieSceneTrack* Track = Params.TrackModel->GetTrack();
-	if (!Track)
-	{
-		return SNullWidget::NullWidget;
-	}
+			FContentBrowserModule& ContentBrowserModule = FModuleManager::Get().LoadModuleChecked<FContentBrowserModule>(TEXT("ContentBrowser"));
 
-	if (ColumnName == FCommonOutlinerNames::Add)
-	{
-		return UE::Sequencer::MakeAddButton(
-			LOCTEXT("AudioControlBusMixText", "Control Bus Mix"),
-			FOnGetContent::CreateSP(
-				this, &FAudioControlBusMixTrackEditor::BuildAudioControlBusMixSubMenu,
-				FOnAssetSelected::CreateRaw(this, &FAudioControlBusMixTrackEditor::OnAudioControlBusMixAssetSelected, Track),
-				FOnAssetEnterPressed::CreateRaw(this, &FAudioControlBusMixTrackEditor::OnAudioControlBusMixAssetEnterPressed, Track)
-			),
-			Params.ViewModel
-		);
-	}
+			const float WidthOverride = Sequencer.IsValid() ? Sequencer->GetSequencerSettings()->GetAssetBrowserWidth() : 500.f;
+			const float HeightOverride = Sequencer.IsValid() ? Sequencer->GetSequencerSettings()->GetAssetBrowserHeight() : 400.f;
 
-	return FMovieSceneTrackEditor::BuildOutlinerColumnWidget(Params, ColumnName);
-}
+			TSharedRef<SWidget> ListWidget = SNew(SBox)
+				.WidthOverride(WidthOverride)
+				.HeightOverride(HeightOverride)
+				[
+					ContentBrowserModule.Get().CreateAssetPicker(AssetPickerConfig)
+				];
 
-bool FAudioControlBusMixTrackEditor::SupportsType(TSubclassOf<UMovieSceneTrack> Type) const
-{
-	return Type == UMovieSceneAudioControlBusMixTrack::StaticClass();
-}
+			SubMenuBuilder.AddWidget(ListWidget, FText::GetEmpty(), true);
+		};
 
-bool FAudioControlBusMixTrackEditor::SupportsSequence(UMovieSceneSequence* InSequence) const
-{
-	return true;
+	MenuBuilder.AddSubMenu(
+		LOCTEXT("AddTrack", "Control Bus Mix Track"),
+		LOCTEXT("AddTooltip", "Adds a new control bus mix track that can toggle a mix on or off"),
+		FNewMenuDelegate::CreateLambda(SubMenuCallback),
+		false,
+		FSlateIconFinder::FindIconForClass(USoundControlBusMix::StaticClass())
+	);
 }
 
-bool FAudioControlBusMixTrackEditor::IsResizable(UMovieSceneTrack* InTrack) const
+bool FAudioControlBusMixTrackEditor::SupportsType(TSubclassOf<UMovieSceneTrack> Type) const
 {
-	return true;
+	return Type == UMovieSceneAudioControlBusMixTrack::StaticClass();
 }
 
-void FAudioControlBusMixTrackEditor::HandleAddAudioControlBusMixTrackMenuEntryExecute()
+void FAudioControlBusMixTrackEditor::AddControlBusAssetTrackToSequence(const FAssetData& InAssetData)
 {
-	if (!IsModulationEnabled())
-	{
-		return;
-	}
-
-	UMovieScene* FocusedMovieScene = GetFocusedMovieScene();
+	FSlateApplication::Get().DismissAllMenus();
 
-	if (FocusedMovieScene == nullptr)
+	USoundControlBusMix* ControlBusMix = Cast<USoundControlBusMix>(InAssetData.GetAsset());
+	UMovieScene* MovieScene = GetFocusedMovieScene();
+	if (!ControlBusMix || !MovieScene)
 	{
 		return;
 	}
 
-	// If movie scene is readonly we cannot add to it, user will have to make it writable and select the add control bus track option again
-	if (FocusedMovieScene->IsReadOnly())
+	if (MovieScene->IsReadOnly())
 	{
 		return;
 	}
 
-	const FScopedTransaction Transaction(NSLOCTEXT("Sequencer", "AddAudioControlBusMixTrack_Transaction", "Add Control Bus Mix Track"));
-	FocusedMovieScene->Modify();
+	const FScopedTransaction Transaction(LOCTEXT("AddAudioControlBusMixTrack_Transaction", "Add Control Bus Mix Track"));
 
-	UMovieSceneAudioControlBusMixTrack* NewTrack = FocusedMovieScene->AddTrack<UMovieSceneAudioControlBusMixTrack>();
-	ensure(NewTrack);
+	MovieScene->Modify();
+	UMovieSceneAudioControlBusMixTrack* NewTrack = MovieScene->AddTrack<UMovieSceneAudioControlBusMixTrack>();
 
 	if (NewTrack == nullptr)
 	{
+		UE_LOG(LogMovieScene, Warning, TEXT("Control Bus Mix Track: New track could not be added for sequence %s"), *MovieScene->GetName());
 		return;
 	}
 
-	NewTrack->SetDisplayName(LOCTEXT("AudioControlBusMixTrackName", "Control Bus Mix"));
+	UMovieSceneSection* NewSection = NewTrack->AddNewControlBusMix(ControlBusMix);
+	check(NewSection);
 
 	if (GetSequencer().IsValid())
 	{
@@ -129,120 +116,20 @@ void FAudioControlBusMixTrackEditor::HandleAddAudioControlBusMixTrackMenuEntryEx
 	}
 }
 
-void FAudioControlBusMixTrackEditor::OnAudioControlBusMixAssetSelected(const FAssetData& AssetData, UMovieSceneTrack* Track)
+bool FAudioControlBusMixTrackEditor::HandleAssetAdded(UObject* Asset, const FGuid& TargetObjectGuid)
 {
-	FSlateApplication::Get().DismissAllMenus();
-
-	if (!IsModulationEnabled())
+	USoundControlBusMix* ControlBus = Cast<USoundControlBusMix>(Asset);
+	if (ControlBus)
 	{
-		return;
+		AddControlBusAssetTrackToSequence(FAssetData(ControlBus));
 	}
 
-	UObject* SelectedObject = AssetData.GetAsset();
-
-	if (SelectedObject)
-	{
-		USoundControlBusMix* NewBusMix = CastChecked<USoundControlBusMix>(AssetData.GetAsset());
-		if (NewBusMix != nullptr)
-		{
-			const FScopedTransaction Transaction(NSLOCTEXT("Sequencer", "AddAudioControlBusMix_Transaction", "Add Control Bus Mix"));
-
-			UMovieSceneAudioControlBusMixTrack* AudioControlBusMixTrack = Cast<UMovieSceneAudioControlBusMixTrack>(Track);
-			check(AudioControlBusMixTrack);
-
-			AudioControlBusMixTrack->Modify();
-
-			check(GetSequencer().IsValid());
-
-			FFrameTime KeyTime = GetSequencer()->GetLocalTime().Time;
-
-			UMovieSceneSection* NewSection = AudioControlBusMixTrack->AddNewControlBusMix(NewBusMix, KeyTime.FrameNumber, INDEX_NONE);
-
-			GetSequencer()->EmptySelection();
-			GetSequencer()->SelectSection(NewSection);
-			GetSequencer()->ThrobSectionSelection();
-
-			GetSequencer()->NotifyMovieSceneDataChanged(EMovieSceneDataChangeType::MovieSceneStructureItemAdded);
-		}
-	}
-}
-
-void FAudioControlBusMixTrackEditor::OnAudioControlBusMixAssetEnterPressed(const TArray<FAssetData>& AssetData, UMovieSceneTrack* Track)
-{
-	// By default enter will select the first list element if there is at least 1 element
-	if (AssetData.Num() > 0)
-	{
-		OnAudioControlBusMixAssetSelected(AssetData[0].GetAsset(), Track);
-	}
-}
-
-TSharedRef<SWidget> FAudioControlBusMixTrackEditor::BuildAudioControlBusMixSubMenu(FOnAssetSelected OnAssetSelected, FOnAssetEnterPressed OnAssetEnterPressed)
-{
-	FMenuBuilder MenuBuilder(true, nullptr);
-
-	if (!IsModulationEnabled())
-	{
-		return MenuBuilder.MakeWidget();
-	}
-
-	const TSharedPtr<ISequencer> SequencerPtr = GetSequencer();
-	const UMovieSceneSequence* Sequence = SequencerPtr.IsValid() ? SequencerPtr->GetFocusedMovieSceneSequence() : nullptr;
-
-	FAssetRegistryModule& AssetRegistryModule = FModuleManager::LoadModuleChecked<FAssetRegistryModule>(TEXT("AssetRegistry"));
-	TArray<FTopLevelAssetPath> ClassNames;
-	ClassNames.Add(USoundControlBusMix::StaticClass()->GetClassPathName());
-
-	TSet<FTopLevelAssetPath> DerivedClassNames;
-	AssetRegistryModule.Get().GetDerivedClassNames(ClassNames, TSet<FTopLevelAssetPath>(), DerivedClassNames);
-
-	FAssetPickerConfig AssetPickerConfig;
-	{
-		AssetPickerConfig.OnAssetSelected = OnAssetSelected;
-		AssetPickerConfig.OnAssetEnterPressed = OnAssetEnterPressed;
-		AssetPickerConfig.bAllowNullSelection = false;
-		AssetPickerConfig.bAddFilterUI = true;
-		AssetPickerConfig.bShowTypeInColumnView = false;
-		AssetPickerConfig.InitialAssetViewType = EAssetViewType::List;
-		for (FTopLevelAssetPath ClassName : DerivedClassNames)
-		{
-			AssetPickerConfig.Filter.ClassPaths.Add(ClassName);
-		}
-		AssetPickerConfig.SaveSettingsName = TEXT("SequencerAssetPicker");
-		AssetPickerConfig.AdditionalReferencingAssets.Add(FAssetData(Sequence));
-	}
-
-	FContentBrowserModule& ContentBrowserModule = FModuleManager::Get().LoadModuleChecked<FContentBrowserModule>(TEXT("ContentBrowser"));
-
-	const float WidthOverride = SequencerPtr.IsValid() ? SequencerPtr->GetSequencerSettings()->GetAssetBrowserWidth() : 500.f;
-	const float HeightOverride = SequencerPtr.IsValid() ? SequencerPtr->GetSequencerSettings()->GetAssetBrowserHeight() : 400.f;
-
-	TSharedPtr<SBox> MenuEntry = SNew(SBox)
-		.WidthOverride(WidthOverride)
-		.HeightOverride(HeightOverride)
-		[
-			ContentBrowserModule.Get().CreateAssetPicker(AssetPickerConfig)
-		];
-
-	MenuBuilder.AddWidget(MenuEntry.ToSharedRef(), FText::GetEmpty(), true);
-
-	return MenuBuilder.MakeWidget();
+	return ControlBus != nullptr;
 }
 
-bool FAudioControlBusMixTrackEditor::IsModulationEnabled()
+const FSlateBrush* FAudioControlBusMixTrackEditor::GetIconBrush() const
 {
-	if (GEditor)
-	{
-		if (UWorld* World = GEditor->GetEditorWorldContext().World())
-		{
-			FAudioDeviceHandle AudioDeviceHandle = World->GetAudioDevice();
-			if (AudioDeviceHandle.IsValid())
-			{
-				return AudioDeviceHandle->IsModulationPluginEnabled();
-			}
-		}
-	}
-
-	return false;
+	return FSlateIconFinder::FindIconForClass(USoundControlBusMix::StaticClass()).GetIcon();
 }
 
 #undef LOCTEXT_NAMESPACE
\ No newline at end of file
diff --git a/Engine/Plugins/Runtime/AudioModulation/Source/AudioModulationEditor/Private/AudioControlBusMixTrackEditor.h b/Engine/Plugins/Runtime/AudioModulation/Source/AudioModulationEditor/Private/AudioControlBusMixTrackEditor.h
--- a/Engine/Plugins/Runtime/AudioModulation/Source/AudioModulationEditor/Private/AudioControlBusMixTrackEditor.h
+++ b/Engine/Plugins/Runtime/AudioModulation/Source/AudioModulationEditor/Private/AudioControlBusMixTrackEditor.h
@@ -2,17 +2,14 @@
 
 #pragma once
 
-#include "IAudioModulation.h"
-#include "IContentBrowserSingleton.h"
+#include "AudioControlBusBaseTrackEditor.h"
 #include "ISequencer.h"
 #include "ISequencerTrackEditor.h"
-#include "MovieSceneSection.h"
-#include "MovieSceneTrackEditor.h"
 
 #define UE_API AUDIOMODULATIONEDITOR_API
 
 class FAudioControlBusMixTrackEditor
-	: public FMovieSceneTrackEditor
+	: public FAudioControlBusBaseTrackEditor
 {
 public:
 	/**
@@ -37,25 +34,14 @@ public:
 
 	UE_API virtual FText GetDisplayName() const override;
 	UE_API virtual void BuildAddTrackMenu(FMenuBuilder& MenuBuilder) override;
-	UE_API virtual TSharedPtr<SWidget> BuildOutlinerColumnWidget(const FBuildColumnWidgetParams& Params, const FName& ColumnName) override;
 	UE_API virtual bool SupportsType(TSubclassOf<UMovieSceneTrack> Type) const override;
-	UE_API virtual bool SupportsSequence(UMovieSceneSequence* InSequence) const override;
-	UE_API virtual bool IsResizable(UMovieSceneTrack* InTrack) const override;
+	UE_API virtual bool HandleAssetAdded(UObject* Asset, const FGuid& TargetObjectGuid) override;
+	UE_API virtual const FSlateBrush* GetIconBrush() const override;
 
 private:
-	/** Callback for executing the "Add Audio Control Bus Mix Track" menu entry. */
-	UE_API void HandleAddAudioControlBusMixTrackMenuEntryExecute();
+	/** Add new control bus mix track to sequence */
+	UE_API void AddControlBusAssetTrackToSequence(const FAssetData& InAssetData) override;
 
-	/** Audio control bus mix asset selected */
-	UE_API void OnAudioControlBusMixAssetSelected(const FAssetData& AssetData, UMovieSceneTrack* Track);
-
-	/** Audio control bus mix asset enter pressed */
-	UE_API void OnAudioControlBusMixAssetEnterPressed(const TArray<FAssetData>& AssetData, UMovieSceneTrack* Track);
-
-	/** Audio control bus mix sub menu */
-	UE_API TSharedRef<SWidget> BuildAudioControlBusMixSubMenu(FOnAssetSelected OnAssetSelected, FOnAssetEnterPressed OnAssetEnterPressed);
-
-	UE_API bool IsModulationEnabled();
 };
 
 #undef UE_API
\ No newline at end of file
diff --git a/Engine/Plugins/Runtime/AudioModulation/Source/AudioModulationEditor/Private/AudioControlBusTrackEditor.cpp b/Engine/Plugins/Runtime/AudioModulation/Source/AudioModulationEditor/Private/AudioControlBusTrackEditor.cpp
--- a/Engine/Plugins/Runtime/AudioModulation/Source/AudioModulationEditor/Private/AudioControlBusTrackEditor.cpp
+++ b/Engine/Plugins/Runtime/AudioModulation/Source/AudioModulationEditor/Private/AudioControlBusTrackEditor.cpp
@@ -2,24 +2,19 @@
 
 #include "AudioControlBusTrackEditor.h"
 
-#include "AssetRegistry/AssetRegistryModule.h"
-#include "AudioDevice.h"
 #include "ContentBrowserModule.h"
-#include "Editor.h"
-#include "Engine/World.h"
 #include "IContentBrowserSingleton.h"
 #include "ISequencer.h"
 #include "ISequencerTrackEditor.h"
-#include "MVVM/Extensions/ITrackExtension.h"
-#include "MVVM/ViewModels/OutlinerColumns/OutlinerColumnTypes.h"
-#include "MVVM/Views/ViewUtilities.h"
-#include "SequencerSettings.h"
 #include "MovieSceneAudioControlBusTrack.h"
+#include "SequencerSettings.h"
+#include "SoundControlBus.h"
+#include "Styling/SlateIconFinder.h"
 
 #define LOCTEXT_NAMESPACE "FAudioControlBusTrackEditor"
 
 FAudioControlBusTrackEditor::FAudioControlBusTrackEditor(TSharedRef<ISequencer> InSequencer)
-	: FMovieSceneTrackEditor(InSequencer)
+	: FAudioControlBusBaseTrackEditor(InSequencer)
 {
 }
 
@@ -39,40 +34,46 @@ FText FAudioControlBusTrackEditor::GetDisplayName() const
 
 void FAudioControlBusTrackEditor::BuildAddTrackMenu(FMenuBuilder& MenuBuilder)
 {
-	MenuBuilder.AddMenuEntry(
-		LOCTEXT("AddTrack", "Control Bus Track"),
-		LOCTEXT("AddTooltip", "Adds a new control track that can be used for mixing"),
-		FSlateIcon(FAppStyle::GetAppStyleSetName(), "Sequencer.Tracks.Audio"),
-		FUIAction(
-			FExecuteAction::CreateRaw(this, &FAudioControlBusTrackEditor::HandleAddAudioControlBusTrackMenuEntryExecute)
-		)
-	);
-}
+	auto SubMenuCallback = [this](FMenuBuilder& SubMenuBuilder)
+		{
+			TSharedPtr<ISequencer> Sequencer = GetSequencer();
+			UMovieSceneSequence* Sequence = Sequencer.IsValid() ? Sequencer->GetFocusedMovieSceneSequence() : nullptr;
 
-TSharedPtr<SWidget> FAudioControlBusTrackEditor::BuildOutlinerColumnWidget(const FBuildColumnWidgetParams& Params, const FName& ColumnName)
-{
-	using namespace UE::Sequencer;
+			FAssetPickerConfig AssetPickerConfig;
+			{
+				AssetPickerConfig.OnAssetSelected = FOnAssetSelected::CreateRaw(this, &FAudioControlBusTrackEditor::AddControlBusAssetTrackToSequence);
+				AssetPickerConfig.OnAssetEnterPressed = FOnAssetEnterPressed::CreateRaw(this, &FAudioControlBusBaseTrackEditor::OnAudioControlBusAssetEnterPressed);
+				AssetPickerConfig.bAllowNullSelection = false;
+				AssetPickerConfig.bAddFilterUI = true;
+				AssetPickerConfig.InitialAssetViewType = EAssetViewType::List;
+				AssetPickerConfig.Filter.bRecursiveClasses = true;
+				AssetPickerConfig.Filter.ClassPaths.Add(USoundControlBus::StaticClass()->GetClassPathName());
+				AssetPickerConfig.SaveSettingsName = TEXT("SequencerAssetPicker");
+				AssetPickerConfig.AdditionalReferencingAssets.Add(FAssetData(Sequence));
+			}
 
-	UMovieSceneTrack* Track = Params.TrackModel->GetTrack();
-	if (!Track)
-	{
-		return SNullWidget::NullWidget;
-	}
+			FContentBrowserModule& ContentBrowserModule = FModuleManager::Get().LoadModuleChecked<FContentBrowserModule>(TEXT("ContentBrowser"));
 
-	if (ColumnName == FCommonOutlinerNames::Add)
-	{
-		return UE::Sequencer::MakeAddButton(
-			LOCTEXT("AudioControlBusText", "Control Bus"),
-			FOnGetContent::CreateSP(
-				this, &FAudioControlBusTrackEditor::BuildAudioControlBusSubMenu,
-				FOnAssetSelected::CreateRaw(this, &FAudioControlBusTrackEditor::OnAudioControlBusAssetSelected, Track),
-				FOnAssetEnterPressed::CreateRaw(this, &FAudioControlBusTrackEditor::OnAudioControlBusAssetEnterPressed, Track)
-			),
-			Params.ViewModel
-		);
-	}
+			const float WidthOverride = Sequencer.IsValid() ? Sequencer->GetSequencerSettings()->GetAssetBrowserWidth() : 500.f;
+			const float HeightOverride = Sequencer.IsValid() ? Sequencer->GetSequencerSettings()->GetAssetBrowserHeight() : 400.f;
+
+			TSharedRef<SWidget> ListWidget = SNew(SBox)
+				.WidthOverride(WidthOverride)
+				.HeightOverride(HeightOverride)
+				[
+					ContentBrowserModule.Get().CreateAssetPicker(AssetPickerConfig)
+				];
 
-	return FMovieSceneTrackEditor::BuildOutlinerColumnWidget(Params, ColumnName);
+			SubMenuBuilder.AddWidget(ListWidget, FText::GetEmpty(), true);
+		};
+
+	MenuBuilder.AddSubMenu(
+		LOCTEXT("AddTrack", "Control Bus Track"),
+		LOCTEXT("AddTooltip", "Adds a new control track that can be used for mixing"),
+		FNewMenuDelegate::CreateLambda(SubMenuCallback),
+		false,
+		FSlateIconFinder::FindIconForClass(USoundControlBus::StaticClass())
+	);
 }
 
 bool FAudioControlBusTrackEditor::SupportsType(TSubclassOf<UMovieSceneTrack> Type) const
@@ -80,179 +81,57 @@ bool FAudioControlBusTrackEditor::SupportsType(TSubclassOf<UMovieSceneTrack> Typ
 	return Type == UMovieSceneAudioControlBusTrack::StaticClass();
 }
 
-bool FAudioControlBusTrackEditor::SupportsSequence(UMovieSceneSequence* InSequence) const
-{
-	return true;
-}
 
-bool FAudioControlBusTrackEditor::IsResizable(UMovieSceneTrack* InTrack) const
+const FSlateBrush* FAudioControlBusTrackEditor::GetIconBrush() const
 {
-	return true;
+	return FSlateIconFinder::FindIconForClass(USoundControlBus::StaticClass()).GetIcon();
 }
 
-void FAudioControlBusTrackEditor::HandleAddAudioControlBusTrackMenuEntryExecute()
+bool FAudioControlBusTrackEditor::HandleAssetAdded(UObject* Asset, const FGuid& TargetObjectGuid)
 {
-	if (!IsModulationEnabled())
-	{
-		return;
-	}
-
-	UMovieScene* FocusedMovieScene = GetFocusedMovieScene();
-
-	if (FocusedMovieScene == nullptr)
-	{
-		return;
-	}
-
-	// If movie scene is readonly we cannot add to it, user will have to make it writable and select the add control bus track option again
-	if (FocusedMovieScene->IsReadOnly())
+	USoundControlBus* ControlBus = Cast<USoundControlBus>(Asset);
+	if (ControlBus)
 	{
-		return;
+		AddControlBusAssetTrackToSequence(FAssetData(ControlBus));
 	}
 
-	const FScopedTransaction Transaction(NSLOCTEXT("Sequencer", "AddAudioControlBusTrack_Transaction", "Add Control Bus Track"));
-	FocusedMovieScene->Modify();
-
-	UMovieSceneAudioControlBusTrack* NewTrack = FocusedMovieScene->AddTrack<UMovieSceneAudioControlBusTrack>();
-	ensure(NewTrack);
-
-	NewTrack->SetDisplayName(LOCTEXT("AudioControlBusTrackName", "Control Bus"));
-
-	if (GetSequencer().IsValid())
-	{
-		GetSequencer()->OnAddTrack(NewTrack, FGuid());
-	}
+	return ControlBus != nullptr;
 }
 
-void FAudioControlBusTrackEditor::OnAudioControlBusAssetSelected(const FAssetData& AssetData, UMovieSceneTrack* Track)
+void FAudioControlBusTrackEditor::AddControlBusAssetTrackToSequence(const FAssetData& InAssetData)
 {
 	FSlateApplication::Get().DismissAllMenus();
 
-	if (!IsModulationEnabled())
+	USoundControlBus* ControlBus = Cast<USoundControlBus>(InAssetData.GetAsset());
+	UMovieScene* MovieScene = GetFocusedMovieScene();
+	if (!ControlBus || !MovieScene)
 	{
 		return;
 	}
 
-	UObject* SelectedObject = AssetData.GetAsset();
-
-	if (SelectedObject)
+	if (MovieScene->IsReadOnly())
 	{
-		USoundModulatorBase* NewBus = CastChecked<USoundModulatorBase>(AssetData.GetAsset());
-		if (NewBus != nullptr)
-		{
-			const FScopedTransaction Transaction(NSLOCTEXT("Sequencer", "AddAudioControlBus_Transaction", "Add Control Bus"));
-
-			UMovieSceneAudioControlBusTrack* AudioControlBusTrack = Cast<UMovieSceneAudioControlBusTrack>(Track);
-			check(AudioControlBusTrack);
-
-			AudioControlBusTrack->Modify();
-
-			check(GetSequencer().IsValid());
-			
-			FFrameTime KeyTime = GetSequencer()->GetLocalTime().Time;
-
-			UMovieSceneSection* NewSection = AudioControlBusTrack->AddNewControlBus(NewBus, KeyTime.FrameNumber, INDEX_NONE);
-
-			GetSequencer()->EmptySelection();
-			GetSequencer()->SelectSection(NewSection);
-			GetSequencer()->ThrobSectionSelection();
-
-			GetSequencer()->NotifyMovieSceneDataChanged(EMovieSceneDataChangeType::MovieSceneStructureItemAdded);
-		}
-	}
-}
-
-void FAudioControlBusTrackEditor::OnAudioControlBusAssetEnterPressed(const TArray<FAssetData>& AssetData, UMovieSceneTrack* Track)
-{
-	// By default enter will select the first list element if there is at least 1 element
-	if (AssetData.Num() > 0)
-	{
-		OnAudioControlBusAssetSelected(AssetData[0].GetAsset(), Track);
-	}
-}
-
-TSharedRef<SWidget> FAudioControlBusTrackEditor::BuildAudioControlBusSubMenu(FOnAssetSelected OnAssetSelected, FOnAssetEnterPressed OnAssetEnterPressed)
-{
-	FMenuBuilder MenuBuilder(true, nullptr);
-
-	if (!IsModulationEnabled())
-	{
-		return MenuBuilder.MakeWidget();
+		return;
 	}
 
-	const TSharedPtr<ISequencer> SequencerPtr = GetSequencer();
-	const UMovieSceneSequence* Sequence = SequencerPtr.IsValid() ? SequencerPtr->GetFocusedMovieSceneSequence() : nullptr;
+	const FScopedTransaction Transaction(LOCTEXT("AddAudioControlBusTrack_Transaction", "Add Control Bus Track"));
 
-	FAssetRegistryModule& AssetRegistryModule = FModuleManager::LoadModuleChecked<FAssetRegistryModule>(TEXT("AssetRegistry"));
-	TArray<FTopLevelAssetPath> ClassNames;
+	MovieScene->Modify();
+	UMovieSceneAudioControlBusTrack* NewTrack = MovieScene->AddTrack<UMovieSceneAudioControlBusTrack>();
 
-	// Due to the plugin storing the USoundControlBus class we have to request the plugin to fetch all the control bus assets since several asset types inherit from USoundModulatorBase
-	if (GEditor)
+	if (NewTrack == nullptr)
 	{
-		if (UWorld* World = GEditor->GetEditorWorldContext().World())
-		{
-			FAudioDeviceHandle AudioDeviceHandle = World->GetAudioDevice();
-			if (AudioDeviceHandle.IsValid() && AudioDeviceHandle->IsModulationPluginEnabled())
-			{
-				if (IAudioModulationManager* Modulation = AudioDeviceHandle->ModulationInterface.Get())
-				{
-					Modulation->GetControlBusAssetPath(ClassNames);
-				}
-			}
-		}
-	}
-
-	TSet<FTopLevelAssetPath> DerivedClassNames;
-	AssetRegistryModule.Get().GetDerivedClassNames(ClassNames, TSet<FTopLevelAssetPath>(), DerivedClassNames);
-
-	FAssetPickerConfig AssetPickerConfig;
-	{
-		AssetPickerConfig.OnAssetSelected = OnAssetSelected;
-		AssetPickerConfig.OnAssetEnterPressed = OnAssetEnterPressed;
-		AssetPickerConfig.bAllowNullSelection = false;
-		AssetPickerConfig.bAddFilterUI = true;
-		AssetPickerConfig.bShowTypeInColumnView = false;
-		AssetPickerConfig.InitialAssetViewType = EAssetViewType::List;
-		for (FTopLevelAssetPath ClassName : DerivedClassNames)
-		{
-			AssetPickerConfig.Filter.ClassPaths.Add(ClassName);
-		}
-		AssetPickerConfig.SaveSettingsName = TEXT("SequencerAssetPicker");
-		AssetPickerConfig.AdditionalReferencingAssets.Add(FAssetData(Sequence));
+		UE_LOG(LogMovieScene, Warning, TEXT("Control Bus Track: New track could not be added for sequence %s"), *MovieScene->GetName());
+		return;
 	}
 
-	FContentBrowserModule& ContentBrowserModule = FModuleManager::Get().LoadModuleChecked<FContentBrowserModule>(TEXT("ContentBrowser"));
-
-	const float WidthOverride = SequencerPtr.IsValid() ? SequencerPtr->GetSequencerSettings()->GetAssetBrowserWidth() : 500.f;
-	const float HeightOverride = SequencerPtr.IsValid() ? SequencerPtr->GetSequencerSettings()->GetAssetBrowserHeight() : 400.f;
-
-	TSharedPtr<SBox> MenuEntry = SNew(SBox)
-		.WidthOverride(WidthOverride)
-		.HeightOverride(HeightOverride)
-		[
-			ContentBrowserModule.Get().CreateAssetPicker(AssetPickerConfig)
-		];
-
-	MenuBuilder.AddWidget(MenuEntry.ToSharedRef(), FText::GetEmpty(), true);
+	UMovieSceneSection* NewSection = NewTrack->AddNewControlBus(ControlBus);
+	check(NewSection);
 
-	return MenuBuilder.MakeWidget();
-}
-
-bool FAudioControlBusTrackEditor::IsModulationEnabled()
-{
-	if (GEditor)
+	if (GetSequencer().IsValid())
 	{
-		if (UWorld* World = GEditor->GetEditorWorldContext().World())
-		{
-			FAudioDeviceHandle AudioDeviceHandle = World->GetAudioDevice();
-			if (AudioDeviceHandle.IsValid())
-			{
-				return AudioDeviceHandle->IsModulationPluginEnabled();
-			}
-		}
+		GetSequencer()->OnAddTrack(NewTrack, FGuid());
 	}
-
-	return false;
 }
 
 #undef LOCTEXT_NAMESPACE
diff --git a/Engine/Plugins/Runtime/AudioModulation/Source/AudioModulationEditor/Private/AudioControlBusTrackEditor.h b/Engine/Plugins/Runtime/AudioModulation/Source/AudioModulationEditor/Private/AudioControlBusTrackEditor.h
--- a/Engine/Plugins/Runtime/AudioModulation/Source/AudioModulationEditor/Private/AudioControlBusTrackEditor.h
+++ b/Engine/Plugins/Runtime/AudioModulation/Source/AudioModulationEditor/Private/AudioControlBusTrackEditor.h
@@ -2,17 +2,14 @@
 
 #pragma once
 
-#include "IAudioModulation.h"
-#include "IContentBrowserSingleton.h"
+#include "AudioControlBusBaseTrackEditor.h"
 #include "ISequencer.h"
 #include "ISequencerTrackEditor.h"
-#include "MovieSceneSection.h"
-#include "MovieSceneTrackEditor.h"
 
 #define UE_API AUDIOMODULATIONEDITOR_API
 
 class FAudioControlBusTrackEditor
-	: public FMovieSceneTrackEditor
+	: public FAudioControlBusBaseTrackEditor
 {
 public:
 	/**
@@ -37,26 +34,13 @@ public:
 
 	UE_API virtual FText GetDisplayName() const override;
 	UE_API virtual void BuildAddTrackMenu(FMenuBuilder& MenuBuilder) override;
-	UE_API virtual TSharedPtr<SWidget> BuildOutlinerColumnWidget(const FBuildColumnWidgetParams& Params, const FName& ColumnName) override;
 	UE_API virtual bool SupportsType(TSubclassOf<UMovieSceneTrack> Type) const override;
-	UE_API virtual bool SupportsSequence(UMovieSceneSequence* InSequence) const override;
-	UE_API virtual bool IsResizable(UMovieSceneTrack* InTrack) const override;
+	UE_API virtual const FSlateBrush* GetIconBrush() const override;
+	UE_API virtual bool HandleAssetAdded(UObject* Asset, const FGuid& TargetObjectGuid) override;
 
 private:
-
-	/** Callback for executing the "Add Audio Control Bus Track" menu entry. */
-	UE_API void HandleAddAudioControlBusTrackMenuEntryExecute();
-
-	/** Audio control bus asset selected */
-	UE_API void OnAudioControlBusAssetSelected(const FAssetData& AssetData, UMovieSceneTrack* Track);
-
-	/** Audio control bus asset enter pressed */
-	UE_API void OnAudioControlBusAssetEnterPressed(const TArray<FAssetData>& AssetData, UMovieSceneTrack* Track);
-
-	/** Audio control bus sub menu */
-	UE_API TSharedRef<SWidget> BuildAudioControlBusSubMenu(FOnAssetSelected OnAssetSelected, FOnAssetEnterPressed OnAssetEnterPressed);
-
-	UE_API bool IsModulationEnabled();
+	/** Add new control bus track to sequence */
+	UE_API void AddControlBusAssetTrackToSequence(const FAssetData& InAssetData) override;
 };
 
 #undef UE_API
\ No newline at end of file
