Sequencer: CVar for control bus tracks

diff --git a/Engine/Plugins/Runtime/AudioModulation/Source/AudioModulationEditor/Private/AudioControlBusBaseTrackEditor.cpp b/Engine/Plugins/Runtime/AudioModulation/Source/AudioModulationEditor/Private/AudioControlBusBaseTrackEditor.cpp
--- a/Engine/Plugins/Runtime/AudioModulation/Source/AudioModulationEditor/Private/AudioControlBusBaseTrackEditor.cpp
+++ b/Engine/Plugins/Runtime/AudioModulation/Source/AudioModulationEditor/Private/AudioControlBusBaseTrackEditor.cpp
@@ -4,9 +4,21 @@
 
 #include "AssetRegistry/AssetData.h"
 
+static bool bUseControlBusTracks = false;
+
+static FAutoConsoleVariableRef CVarControlBusTracks(
+	TEXT("au.ControlBusTracks"),
+	bUseControlBusTracks,
+	TEXT("Enable the control bus track and control bus mix track in sequencer using the audio modulation plugin. Default disables adding new control bus tracks\n")
+	TEXT("0: Disable control bus tracks\n")
+	TEXT("1: Enable control bus tracks\n"),
+	ECVF_Default
+);
+
 FAudioControlBusBaseTrackEditor::FAudioControlBusBaseTrackEditor(TSharedRef<ISequencer> InSequencer)
 	: FMovieSceneTrackEditor(InSequencer)
 {
+	CVarUseControlBusTracks = IConsoleManager::Get().FindConsoleVariable(TEXT("au.ControlBusTracks"));
 }
 
 FAudioControlBusBaseTrackEditor::~FAudioControlBusBaseTrackEditor()
diff --git a/Engine/Plugins/Runtime/AudioModulation/Source/AudioModulationEditor/Private/AudioControlBusBaseTrackEditor.h b/Engine/Plugins/Runtime/AudioModulation/Source/AudioModulationEditor/Private/AudioControlBusBaseTrackEditor.h
--- a/Engine/Plugins/Runtime/AudioModulation/Source/AudioModulationEditor/Private/AudioControlBusBaseTrackEditor.h
+++ b/Engine/Plugins/Runtime/AudioModulation/Source/AudioModulationEditor/Private/AudioControlBusBaseTrackEditor.h
@@ -30,6 +30,8 @@ public:
 protected:
 	/** Add new control bus Asset track to sequence */
 	UE_API virtual void AddControlBusAssetTrackToSequence(const FAssetData& InAssetData);
+
+	IConsoleVariable* CVarUseControlBusTracks;
 };
 
 #undef UE_API
diff --git a/Engine/Plugins/Runtime/AudioModulation/Source/AudioModulationEditor/Private/AudioControlBusMixTrackEditor.cpp b/Engine/Plugins/Runtime/AudioModulation/Source/AudioModulationEditor/Private/AudioControlBusMixTrackEditor.cpp
--- a/Engine/Plugins/Runtime/AudioModulation/Source/AudioModulationEditor/Private/AudioControlBusMixTrackEditor.cpp
+++ b/Engine/Plugins/Runtime/AudioModulation/Source/AudioModulationEditor/Private/AudioControlBusMixTrackEditor.cpp
@@ -33,6 +33,11 @@ FText FAudioControlBusMixTrackEditor::GetDisplayName() const
 
 void FAudioControlBusMixTrackEditor::BuildAddTrackMenu(FMenuBuilder& MenuBuilder)
 {
+	if (!CVarUseControlBusTracks || !CVarUseControlBusTracks->GetBool())
+	{
+		return;
+	}
+
 	auto SubMenuCallback = [this](FMenuBuilder& SubMenuBuilder)
 		{
 			TSharedPtr<ISequencer> Sequencer = GetSequencer();
@@ -82,6 +87,11 @@ bool FAudioControlBusMixTrackEditor::SupportsType(TSubclassOf<UMovieSceneTrack>
 
 void FAudioControlBusMixTrackEditor::AddControlBusAssetTrackToSequence(const FAssetData& InAssetData)
 {
+	if (!CVarUseControlBusTracks || !CVarUseControlBusTracks->GetBool())
+	{
+		return;
+	}
+
 	FSlateApplication::Get().DismissAllMenus();
 
 	USoundControlBusMix* ControlBusMix = Cast<USoundControlBusMix>(InAssetData.GetAsset());
@@ -118,6 +128,11 @@ void FAudioControlBusMixTrackEditor::AddControlBusAssetTrackToSequence(const FAs
 
 bool FAudioControlBusMixTrackEditor::HandleAssetAdded(UObject* Asset, const FGuid& TargetObjectGuid)
 {
+	if (!CVarUseControlBusTracks || !CVarUseControlBusTracks->GetBool())
+	{
+		return false;
+	}
+
 	USoundControlBusMix* ControlBus = Cast<USoundControlBusMix>(Asset);
 	if (ControlBus)
 	{
diff --git a/Engine/Plugins/Runtime/AudioModulation/Source/AudioModulationEditor/Private/AudioControlBusTrackEditor.cpp b/Engine/Plugins/Runtime/AudioModulation/Source/AudioModulationEditor/Private/AudioControlBusTrackEditor.cpp
--- a/Engine/Plugins/Runtime/AudioModulation/Source/AudioModulationEditor/Private/AudioControlBusTrackEditor.cpp
+++ b/Engine/Plugins/Runtime/AudioModulation/Source/AudioModulationEditor/Private/AudioControlBusTrackEditor.cpp
@@ -34,6 +34,11 @@ FText FAudioControlBusTrackEditor::GetDisplayName() const
 
 void FAudioControlBusTrackEditor::BuildAddTrackMenu(FMenuBuilder& MenuBuilder)
 {
+	if (!CVarUseControlBusTracks || !CVarUseControlBusTracks->GetBool())
+	{
+		return;
+	}
+
 	auto SubMenuCallback = [this](FMenuBuilder& SubMenuBuilder)
 		{
 			TSharedPtr<ISequencer> Sequencer = GetSequencer();
@@ -89,6 +94,11 @@ const FSlateBrush* FAudioControlBusTrackEditor::GetIconBrush() const
 
 bool FAudioControlBusTrackEditor::HandleAssetAdded(UObject* Asset, const FGuid& TargetObjectGuid)
 {
+	if (!CVarUseControlBusTracks || !CVarUseControlBusTracks->GetBool())
+	{
+		return false;
+	}
+
 	USoundControlBus* ControlBus = Cast<USoundControlBus>(Asset);
 	if (ControlBus)
 	{
@@ -100,6 +110,11 @@ bool FAudioControlBusTrackEditor::HandleAssetAdded(UObject* Asset, const FGuid&
 
 void FAudioControlBusTrackEditor::AddControlBusAssetTrackToSequence(const FAssetData& InAssetData)
 {
+	if (!CVarUseControlBusTracks || !CVarUseControlBusTracks->GetBool())
+	{
+		return;
+	}
+
 	FSlateApplication::Get().DismissAllMenus();
 
 	USoundControlBus* ControlBus = Cast<USoundControlBus>(InAssetData.GetAsset());
