Sequencer: Hide control bus tracks in actor sequence component

#Synced-CL 48625258

diff --git a/Engine/Plugins/Runtime/AudioModulation/Source/AudioModulationEditor/AudioModulationEditor.build.cs b/Engine/Plugins/Runtime/AudioModulation/Source/AudioModulationEditor/AudioModulationEditor.build.cs
--- a/Engine/Plugins/Runtime/AudioModulation/Source/AudioModulationEditor/AudioModulationEditor.build.cs
+++ b/Engine/Plugins/Runtime/AudioModulation/Source/AudioModulationEditor/AudioModulationEditor.build.cs
@@ -37,7 +37,8 @@ public class AudioModulationEditor : ModuleRules
 				"MovieScene",
 				"Sequencer",
 				"SequencerCore",
-				"MovieSceneTools"
+				"MovieSceneTools",
+				"LevelSequence"
 			}
 		);
 
diff --git a/Engine/Plugins/Runtime/AudioModulation/Source/AudioModulationEditor/Private/AudioControlBusMixTrackEditor.cpp b/Engine/Plugins/Runtime/AudioModulation/Source/AudioModulationEditor/Private/AudioControlBusMixTrackEditor.cpp
--- a/Engine/Plugins/Runtime/AudioModulation/Source/AudioModulationEditor/Private/AudioControlBusMixTrackEditor.cpp
+++ b/Engine/Plugins/Runtime/AudioModulation/Source/AudioModulationEditor/Private/AudioControlBusMixTrackEditor.cpp
@@ -6,6 +6,7 @@
 #include "IContentBrowserSingleton.h"
 #include "ISequencer.h"
 #include "ISequencerTrackEditor.h"
+#include "LevelSequence.h"
 #include "MovieSceneAudioControlBusMixTrack.h"
 #include "SequencerSettings.h"
 #include "Styling/SlateIconFinder.h"
@@ -85,6 +86,16 @@ bool FAudioControlBusMixTrackEditor::SupportsType(TSubclassOf<UMovieSceneTrack>
 	return Type == UMovieSceneAudioControlBusMixTrack::StaticClass();
 }
 
+bool FAudioControlBusMixTrackEditor::SupportsSequence(UMovieSceneSequence* InSequence) const
+{
+	if (InSequence && InSequence->IsTrackSupported(UMovieSceneAudioControlBusMixTrack::StaticClass()) == ETrackSupport::NotSupported)
+	{
+		return false;
+	}
+
+	return InSequence && InSequence->IsA(ULevelSequence::StaticClass());
+}
+
 void FAudioControlBusMixTrackEditor::AddControlBusAssetTrackToSequence(const FAssetData& InAssetData)
 {
 	if (!CVarUseControlBusTracks || !CVarUseControlBusTracks->GetBool())
diff --git a/Engine/Plugins/Runtime/AudioModulation/Source/AudioModulationEditor/Private/AudioControlBusMixTrackEditor.h b/Engine/Plugins/Runtime/AudioModulation/Source/AudioModulationEditor/Private/AudioControlBusMixTrackEditor.h
--- a/Engine/Plugins/Runtime/AudioModulation/Source/AudioModulationEditor/Private/AudioControlBusMixTrackEditor.h
+++ b/Engine/Plugins/Runtime/AudioModulation/Source/AudioModulationEditor/Private/AudioControlBusMixTrackEditor.h
@@ -35,6 +35,7 @@ public:
 	UE_API virtual FText GetDisplayName() const override;
 	UE_API virtual void BuildAddTrackMenu(FMenuBuilder& MenuBuilder) override;
 	UE_API virtual bool SupportsType(TSubclassOf<UMovieSceneTrack> Type) const override;
+	UE_API virtual bool SupportsSequence(UMovieSceneSequence* InSequence) const override;
 	UE_API virtual bool HandleAssetAdded(UObject* Asset, const FGuid& TargetObjectGuid) override;
 	UE_API virtual const FSlateBrush* GetIconBrush() const override;
 
diff --git a/Engine/Plugins/Runtime/AudioModulation/Source/AudioModulationEditor/Private/AudioControlBusTrackEditor.cpp b/Engine/Plugins/Runtime/AudioModulation/Source/AudioModulationEditor/Private/AudioControlBusTrackEditor.cpp
--- a/Engine/Plugins/Runtime/AudioModulation/Source/AudioModulationEditor/Private/AudioControlBusTrackEditor.cpp
+++ b/Engine/Plugins/Runtime/AudioModulation/Source/AudioModulationEditor/Private/AudioControlBusTrackEditor.cpp
@@ -6,6 +6,7 @@
 #include "IContentBrowserSingleton.h"
 #include "ISequencer.h"
 #include "ISequencerTrackEditor.h"
+#include "LevelSequence.h"
 #include "MovieSceneAudioControlBusTrack.h"
 #include "SequencerSettings.h"
 #include "SoundControlBus.h"
@@ -86,6 +87,15 @@ bool FAudioControlBusTrackEditor::SupportsType(TSubclassOf<UMovieSceneTrack> Typ
 	return Type == UMovieSceneAudioControlBusTrack::StaticClass();
 }
 
+bool FAudioControlBusTrackEditor::SupportsSequence(UMovieSceneSequence* InSequence) const
+{
+	if (InSequence && InSequence->IsTrackSupported(UMovieSceneAudioControlBusTrack::StaticClass()) == ETrackSupport::NotSupported)
+	{
+		return false;
+	}
+
+	return InSequence && InSequence->IsA(ULevelSequence::StaticClass());
+}
 
 const FSlateBrush* FAudioControlBusTrackEditor::GetIconBrush() const
 {
diff --git a/Engine/Plugins/Runtime/AudioModulation/Source/AudioModulationEditor/Private/AudioControlBusTrackEditor.h b/Engine/Plugins/Runtime/AudioModulation/Source/AudioModulationEditor/Private/AudioControlBusTrackEditor.h
--- a/Engine/Plugins/Runtime/AudioModulation/Source/AudioModulationEditor/Private/AudioControlBusTrackEditor.h
+++ b/Engine/Plugins/Runtime/AudioModulation/Source/AudioModulationEditor/Private/AudioControlBusTrackEditor.h
@@ -35,6 +35,7 @@ public:
 	UE_API virtual FText GetDisplayName() const override;
 	UE_API virtual void BuildAddTrackMenu(FMenuBuilder& MenuBuilder) override;
 	UE_API virtual bool SupportsType(TSubclassOf<UMovieSceneTrack> Type) const override;
+	UE_API virtual bool SupportsSequence(UMovieSceneSequence* InSequence) const override;
 	UE_API virtual const FSlateBrush* GetIconBrush() const override;
 	UE_API virtual bool HandleAssetAdded(UObject* Asset, const FGuid& TargetObjectGuid) override;
 
