[Audio] Make ControlBusMix details show control bus names as TitleProperty for each mix stage (tooltip shows complete pathname)

#Synced-CL 49677159

diff --git a/Engine/Plugins/Runtime/AudioModulation/Source/AudioModulationEditor/Private/Layouts/SoundControlBusMixStageLayout.cpp b/Engine/Plugins/Runtime/AudioModulation/Source/AudioModulationEditor/Private/Layouts/SoundControlBusMixStageLayout.cpp
--- a/Engine/Plugins/Runtime/AudioModulation/Source/AudioModulationEditor/Private/Layouts/SoundControlBusMixStageLayout.cpp
+++ b/Engine/Plugins/Runtime/AudioModulation/Source/AudioModulationEditor/Private/Layouts/SoundControlBusMixStageLayout.cpp
@@ -150,8 +150,43 @@ void FSoundControlBusMixStageLayoutCustomization::CustomizeHeader(TSharedRef<IPr
 			[
 				SNew(STextBlock)
 				.Font(IDetailLayoutBuilder::GetDetailFont())
-				.Text(FText::Format(LOCTEXT("BusMixStageHeader_Format", "Bus Stage {0}"), StructPropertyHandle->GetPropertyDisplayName()))
-				.ToolTipText(StructPropertyHandle->GetToolTipText())
+				.Text_Lambda(
+					[StructPropertyHandle]()
+					{
+						constexpr bool bRecurse = false;
+						if (TSharedPtr<IPropertyHandle> BusProperty = StructPropertyHandle->GetChildHandle(GET_MEMBER_NAME_CHECKED(FSoundControlBusMixStage, Bus), bRecurse))
+						{
+							const UObject* BusObject = nullptr;
+							if (BusProperty->GetValue(BusObject) == FPropertyAccess::Success)
+							{
+								return FText::AsCultureInvariant(GetNameSafe(BusObject));
+							}
+						}
+
+						return FText::Format(LOCTEXT("BusMixStageHeader_Format", "Bus Stage {0}"), StructPropertyHandle->GetPropertyDisplayName());
+					}
+				)
+				.ToolTipText_Lambda(
+					[StructPropertyHandle]()
+					{
+						constexpr bool bRecurse = false;
+						if (TSharedPtr<IPropertyHandle> BusProperty = StructPropertyHandle->GetChildHandle(GET_MEMBER_NAME_CHECKED(FSoundControlBusMixStage, Bus), bRecurse))
+						{
+							const UObject* BusObject = nullptr;
+							if (BusProperty->GetValue(BusObject) == FPropertyAccess::Success)
+							{
+								if (BusObject != nullptr)
+								{
+									UPackage* Package = BusObject->GetPackage();
+									return FText::AsCultureInvariant(GetPathNameSafe(Package ? (UObject*)Package : BusObject));
+								}
+								return FText::GetEmpty();
+							}
+						}
+
+						return StructPropertyHandle->GetToolTipText();
+					}
+				)
 			]
 		];
 }
