Fix Parameter Unclamp CVar Not working when setting Control Bus Mix Values
#Synced-CL 49661686

diff --git a/Engine/Plugins/Runtime/AudioModulation/Source/AudioModulation/Private/SoundControlBusMix.cpp b/Engine/Plugins/Runtime/AudioModulation/Source/AudioModulation/Private/SoundControlBusMix.cpp
--- a/Engine/Plugins/Runtime/AudioModulation/Source/AudioModulation/Private/SoundControlBusMix.cpp
+++ b/Engine/Plugins/Runtime/AudioModulation/Source/AudioModulation/Private/SoundControlBusMix.cpp
@@ -135,8 +135,9 @@ void USoundControlBusMix::OnPropertyChanged(FProperty* Property, EPropertyChange
 						float UnitValue = Stage.Value.TargetValue;
 						if (USoundModulationParameter* Parameter = Stage.Bus->Parameter)
 						{
-							// Clamp the values--either to the normalized [0,1] range or to the parameter-specified "unclamped" range
-							if (!Parameter->Settings.bClampToNormalizedRange)
+							// Clamp the normalized values--either to the normalized [0,1] range or to the parameter-specified "unclamped" range.
+							// The normalized value is not clamped in the ConvertNormalizedToUnit function below, so we do it here.
+							if (!Parameter->IsModulationParameterClameped())
 							{
 								const TRange<float> Range = Parameter->GetNormalizedLimits();
 								Stage.Value.TargetValue = FMath::Clamp(Stage.Value.TargetValue, Range.GetLowerBoundValue(), Range.GetUpperBoundValue());
diff --git a/Engine/Plugins/Runtime/AudioModulation/Source/AudioModulation/Private/SoundModulationParameter.cpp b/Engine/Plugins/Runtime/AudioModulation/Source/AudioModulation/Private/SoundModulationParameter.cpp
--- a/Engine/Plugins/Runtime/AudioModulation/Source/AudioModulation/Private/SoundModulationParameter.cpp
+++ b/Engine/Plugins/Runtime/AudioModulation/Source/AudioModulation/Private/SoundModulationParameter.cpp
@@ -80,6 +80,14 @@ Audio::FModulationParameter USoundModulationParameter::CreateParameter() const
 	return Parameter;
 }
 
+bool USoundModulationParameter::IsModulationParameterClameped() const
+{
+	bool bClampToNormalizedRange = Settings.bClampToNormalizedRange;
+	Audio::GetModulationParameterClampValue(GetFName(), bClampToNormalizedRange);
+
+	return bClampToNormalizedRange;
+}
+
 bool USoundModulationParameterFrequencyBase::RequiresUnitConversion() const
 {
 	return true;
@@ -264,7 +272,7 @@ bool USoundModulationParameterVolume::RequiresUnitConversion() const
 
 Audio::FModulationUnitConversionFunction USoundModulationParameterVolume::GetUnitConversionFunction() const
 {
-	return [InUnitMin = GetUnitMin(), InUnitMax = VolumeLimit, bClampToNormalized = Settings.bClampToNormalizedRange](float& InOutValue)
+	return [InUnitMin = GetUnitMin(), InUnitMax = VolumeLimit, bClampToNormalized = IsModulationParameterClameped()](float& InOutValue)
 	{
 		InOutValue = InOutValue > 0.0f
 			? Audio::ConvertToDecibels(InOutValue)
@@ -283,7 +291,7 @@ Audio::FModulationUnitConversionFunction USoundModulationParameterVolume::GetUni
 
 Audio::FModulationNormalizedConversionFunction USoundModulationParameterVolume::GetNormalizedConversionFunction() const
 {
-	return [InUnitMin = GetUnitMin(), InUnitMax = VolumeLimit, bClampToNormalized = Settings.bClampToNormalizedRange](float& InOutValue)
+	return [InUnitMin = GetUnitMin(), InUnitMax = VolumeLimit, bClampToNormalized = IsModulationParameterClameped()](float& InOutValue)
 	{
 		InOutValue = InOutValue < InUnitMin || FMath::IsNearlyEqual(InOutValue, InUnitMin)
 			? 0.0f
diff --git a/Engine/Plugins/Runtime/AudioModulation/Source/AudioModulation/Public/AudioModulationSubsystem.h b/Engine/Plugins/Runtime/AudioModulation/Source/AudioModulation/Public/AudioModulationSubsystem.h
--- a/Engine/Plugins/Runtime/AudioModulation/Source/AudioModulation/Public/AudioModulationSubsystem.h
+++ b/Engine/Plugins/Runtime/AudioModulation/Source/AudioModulation/Public/AudioModulationSubsystem.h
@@ -6,6 +6,8 @@
 
 #include "AudioModulationSubsystem.generated.h"
 
+#define UE_API AUDIOMODULATION_API
+
 /**
 *  UAudioModulationSubsystem
 */
@@ -19,8 +21,8 @@ public:
 	virtual ~UAudioModulationSubsystem() override = default;
 
 	//~ Begin USubsystem interface
-	AUDIOMODULATION_API virtual bool ShouldCreateSubsystem(UObject* Outer) const override;
-	AUDIOMODULATION_API virtual void Deinitialize() override;
+	UE_API virtual bool ShouldCreateSubsystem(UObject* Outer) const override;
+	UE_API virtual void Deinitialize() override;
 	//~ End USubsystem interface
 
 	// Modulation Parameter Overrides
@@ -31,3 +33,5 @@ public:
 private:
 	TMap<FString, bool> ModulationParameterClampOverrides;
 };
+
+# undef UE_API
diff --git a/Engine/Plugins/Runtime/AudioModulation/Source/AudioModulation/Public/SoundModulationParameter.h b/Engine/Plugins/Runtime/AudioModulation/Source/AudioModulation/Public/SoundModulationParameter.h
--- a/Engine/Plugins/Runtime/AudioModulation/Source/AudioModulation/Public/SoundModulationParameter.h
+++ b/Engine/Plugins/Runtime/AudioModulation/Source/AudioModulation/Public/SoundModulationParameter.h
@@ -130,6 +130,8 @@ public:
 #endif // WITH_EDITOR
 
 	UE_API Audio::FModulationParameter CreateParameter() const;
+
+	UE_API bool IsModulationParameterClameped() const;
 };
 
 // Linearly scaled value between unit minimum and maximum.
diff --git a/Engine/Source/Runtime/AudioExtensions/Private/IAudioModulation.cpp b/Engine/Source/Runtime/AudioExtensions/Private/IAudioModulation.cpp
--- a/Engine/Source/Runtime/AudioExtensions/Private/IAudioModulation.cpp
+++ b/Engine/Source/Runtime/AudioExtensions/Private/IAudioModulation.cpp
@@ -120,6 +120,21 @@ namespace Audio
 
 				return nullptr;
 			}
+
+			bool GetParameterClamp(FName InName) const
+			{
+				FScopeLock Lock(&ThreadSafeValueAccessor);
+
+				if (const TUniquePtr<FParameterRegisterEntry>* ParamEntry = Values.Find(InName))
+				{
+					if (FParameterRegisterEntry* ParamEntryPtr = ParamEntry->Get())
+					{
+						return ParamEntryPtr->Parameter.bClampToNormalizedRange;
+					}
+				}
+
+				return true;
+			}
 		} ParameterRegistry;
 	}
 
@@ -213,6 +228,17 @@ namespace Audio
 		ModulationInterfacePrivate::ParameterRegistry.ClearAllClampOverrides();
 	}
 
+	bool GetModulationParameterClampValue(const FName& ParameterName, bool& bClampOut)
+	{
+		if (ModulationInterfacePrivate::ParameterRegistry.IsRegistered(ParameterName))
+		{
+			bClampOut = ModulationInterfacePrivate::ParameterRegistry.GetParameterClamp(ParameterName);
+			return true;
+		}
+
+		return false;
+	}
+
 	const FModulationParameter& GetModulationParameter(FName InName)
 	{
 		using namespace ModulationInterfacePrivate;
diff --git a/Engine/Source/Runtime/AudioExtensions/Public/IAudioModulation.h b/Engine/Source/Runtime/AudioExtensions/Public/IAudioModulation.h
--- a/Engine/Source/Runtime/AudioExtensions/Public/IAudioModulation.h
+++ b/Engine/Source/Runtime/AudioExtensions/Public/IAudioModulation.h
@@ -116,6 +116,10 @@ namespace Audio
 	
 	UE_INTERNAL
 	AUDIOEXTENSIONS_API void ClearAllModulationParameterClampOverrides();
+
+	// Returns false if parameter is not registered
+	UE_INTERNAL
+	AUDIOEXTENSIONS_API bool GetModulationParameterClampValue(const FName& ParameterName, bool& bClampOut);
 	
 	UE_DEPRECATED(5.6, "Please use the version of GetModulationParameter which returns a pointer instead of a reference, and check that it is valid in your code.")
 	AUDIOEXTENSIONS_API const FModulationParameter& GetModulationParameter(FName InName);
