[AudioInsights] - AudioModulation:  Parameter Patch, promote local variable ModStageValue to member

Partial fix for:
https://jira.it.epicgames.com/browse/UE-355680

#Synced-CL 49728086

diff --git a/Engine/Plugins/Runtime/AudioModulation/Source/AudioModulation/Private/SoundModulationPatchProxy.cpp b/Engine/Plugins/Runtime/AudioModulation/Source/AudioModulation/Private/SoundModulationPatchProxy.cpp
--- a/Engine/Plugins/Runtime/AudioModulation/Source/AudioModulation/Private/SoundModulationPatchProxy.cpp
+++ b/Engine/Plugins/Runtime/AudioModulation/Source/AudioModulation/Private/SoundModulationPatchProxy.cpp
@@ -27,6 +27,11 @@ namespace AudioModulation
 		, Transform(MoveTemp(InSettings.Transform))
 		, bSampleAndHold(InSettings.bSampleAndHold)
 	{
+		if (BusHandle.IsValid())
+		{
+			const FControlBusProxy& BusProxy = BusHandle.FindProxy();
+			ModStageValue = BusProxy.GetValue();
+		}
 	}
 
 	FModulationOutputProxy::FModulationOutputProxy(float InDefaultValue, const Audio::FModulationMixFunction& InMixFunction)
@@ -87,7 +92,7 @@ namespace AudioModulation
 			OutputProxy.bInitialized = true;
 		}
 
-		for (const FModulationInputProxy& Input : InputProxies)
+		for (FModulationInputProxy& Input : InputProxies)
 		{
 			if (Input.bSampleAndHold)
 			{
@@ -96,9 +101,9 @@ namespace AudioModulation
 					const FControlBusProxy& BusProxy = Input.BusHandle.FindProxy();
 					if (!BusProxy.IsBypassed())
 					{
-						float ModStageValue = BusProxy.GetValue();
-						Input.Transform.Apply(ModStageValue);
-						OutputProxy.MixFunction(OutSampleHold, ModStageValue);
+						Input.ModStageValue = BusProxy.GetValue();
+						Input.Transform.Apply(Input.ModStageValue);
+						OutputProxy.MixFunction(OutSampleHold, Input.ModStageValue);
 					}
 				}
 			}
@@ -109,9 +114,9 @@ namespace AudioModulation
 					const FControlBusProxy& BusProxy = Input.BusHandle.FindProxy();
 					if (!BusProxy.IsBypassed())
 					{
-						float ModStageValue = BusProxy.GetValue();
-						Input.Transform.Apply(ModStageValue);
-						OutputProxy.MixFunction(Value, ModStageValue);
+						Input.ModStageValue = BusProxy.GetValue();
+						Input.Transform.Apply(Input.ModStageValue);
+						OutputProxy.MixFunction(Value, Input.ModStageValue);
 					}
 				}
 			}
diff --git a/Engine/Plugins/Runtime/AudioModulation/Source/AudioModulation/Private/SoundModulationPatchProxy.h b/Engine/Plugins/Runtime/AudioModulation/Source/AudioModulation/Private/SoundModulationPatchProxy.h
--- a/Engine/Plugins/Runtime/AudioModulation/Source/AudioModulation/Private/SoundModulationPatchProxy.h
+++ b/Engine/Plugins/Runtime/AudioModulation/Source/AudioModulation/Private/SoundModulationPatchProxy.h
@@ -50,6 +50,9 @@ namespace AudioModulation
 		FBusHandle BusHandle;
 
 		FSoundModulationTransform Transform;
+
+		float ModStageValue = 1.0f;
+
 		bool bSampleAndHold = false;
 	};
 
