Sequencer: Control bus track reference fixes

#Synced-CL 50507089

diff --git a/Engine/Plugins/Runtime/AudioModulation/Source/AudioModulation/Private/AudioModulation.cpp b/Engine/Plugins/Runtime/AudioModulation/Source/AudioModulation/Private/AudioModulation.cpp
--- a/Engine/Plugins/Runtime/AudioModulation/Source/AudioModulation/Private/AudioModulation.cpp
+++ b/Engine/Plugins/Runtime/AudioModulation/Source/AudioModulation/Private/AudioModulation.cpp
@@ -43,7 +43,6 @@ namespace AudioModulation
 	FAudioModulationManager::FAudioModulationManager()
 		: ModSystem(new FAudioModulationSystem())
 	{
-
 	}
 
 	FAudioModulationManager::~FAudioModulationManager()
@@ -66,70 +65,6 @@ namespace AudioModulation
 		OutClassNames.Add(USoundControlBus::StaticClass()->GetClassPathName());
 	}
 
-	void FAudioModulationManager::CreateNewSequencerBusMix(const UObject* WorldContextObject, FName Name, const TArray<USoundModulatorBase*>& Buses, float Value, float AttackTime, float ReleaseTime, bool bActivate)
-	{
-		UWorld* World = UAudioModulationStatics::GetAudioWorld(WorldContextObject);
-		TArray<USoundControlBus*> BusArray;
-
-		for (USoundModulatorBase* Bus : Buses)
-		{
-			USoundControlBus* CastBus = Cast<USoundControlBus>(Bus);
-
-			if (CastBus != nullptr)
-			{
-				BusArray.Add(CastBus);
-			}
-		}
-
-		if (USoundControlBusMix* NewMix = CreateBusMixFromValue(Name, BusArray, Value, AttackTime, ReleaseTime))
-		{
-			// If we already have an entry with this ID for some reason, clear it to avoid untracked handles before adding the new entry
-			if (SequencerBusMixes.Contains(Name))
-			{
-				ClearSequencerBusMix(Name);
-			}
-
-			if (bActivate)
-			{
-				ModSystem->ActivateBusMix(*NewMix);
-			}
-
-			SequencerBusMixes.Add(Name, NewMix);
-		}
-	}
-
-	void FAudioModulationManager::UpdateMixByModulator(const UObject* WorldContextObject, FName Mix, USoundModulatorBase& Modulator, float Value, float FadeTime)
-	{
-		if (SequencerBusMixes.Contains(Mix))
-		{
-			USoundControlBusMix* ControlBusMix = SequencerBusMixes[Mix];
-			USoundControlBus* CastBus = Cast<USoundControlBus>(&Modulator);
-
-			check(ControlBusMix);
-			check(CastBus);
-
-			if (CastBus == nullptr || ControlBusMix == nullptr)
-			{
-				return;
-			}
-
-			UpdateMixByFilter(CastBus->Address, nullptr, CastBus->Parameter, Value, FadeTime, *ControlBusMix);
-		}
-	}
-
-	void FAudioModulationManager::ClearSequencerBusMix(FName Name)
-	{
-		if (SequencerBusMixes.Contains(Name))
-		{
-			USoundControlBusMix* Mix = SequencerBusMixes[Name];
-			check(Mix);
-
-			ModSystem->DeactivateBusMix(*Mix);
-
-			SequencerBusMixes.Remove(Name);
-		}
-	}
-
 	void FAudioModulationManager::ActivateBus(const USoundControlBus& InBus)
 	{
 		UE_LOG(LogAudioModulation, VeryVerbose, TEXT("FAudioModulationManager::ActivateBus is deprecated. Use UAudioModulationDestination API (see SetModulator) instead to enforce bus lifetime"));
@@ -172,6 +107,11 @@ namespace AudioModulation
 		ModSystem->DeactivateBusMix(InBusMix);
 	}
 
+	void FAudioModulationManager::DeactivateBusMixByID(const uint32 ID)
+	{
+		ModSystem->DeactivateBusMixByID(ID);
+	}
+
 	void FAudioModulationManager::DeactivateAllBusMixes()
 	{
 		ModSystem->DeactivateAllBusMixes();
diff --git a/Engine/Plugins/Runtime/AudioModulation/Source/AudioModulation/Private/AudioModulationStatics.cpp b/Engine/Plugins/Runtime/AudioModulation/Source/AudioModulation/Private/AudioModulationStatics.cpp
--- a/Engine/Plugins/Runtime/AudioModulation/Source/AudioModulation/Private/AudioModulationStatics.cpp
+++ b/Engine/Plugins/Runtime/AudioModulation/Source/AudioModulation/Private/AudioModulationStatics.cpp
@@ -389,6 +389,15 @@ void UAudioModulationStatics::DeactivateBusMix(const UObject* WorldContextObject
 	}
 }
 
+void UAudioModulationStatics::DeactivateBusMixByID(const UObject* WorldContextObject, uint32 MixID)
+{
+	UWorld* World = GetAudioWorld(WorldContextObject);
+	if (AudioModulation::FAudioModulationManager* ModSystem = GetModulation(World))
+	{
+		ModSystem->DeactivateBusMixByID(MixID);
+	}
+}
+
 void UAudioModulationStatics::DeactivateGenerator(const UObject* WorldContextObject, USoundModulationGenerator* Generator)
 {
 	UWorld* World = GetAudioWorld(WorldContextObject);
@@ -463,6 +472,7 @@ void UAudioModulationStatics::UpdateMix(const UObject* WorldContextObject, USoun
 USoundControlBusMix* UAudioModulationStatics::CreateBusMixFromValue(const UObject* WorldContextObject, FName Name, const TArray<USoundControlBus*>& Buses, float Value, float AttackTime, float ReleaseTime, bool bActivate)
 {
 	UWorld* World = GetAudioWorld(WorldContextObject);
+
 	if (AudioModulation::FAudioModulationManager* ModSystem = GetModulation(World))
 	{
 		if (USoundControlBusMix* NewMix = ModSystem->CreateBusMixFromValue(Name, Buses, Value, AttackTime, ReleaseTime))
diff --git a/Engine/Plugins/Runtime/AudioModulation/Source/AudioModulation/Private/AudioModulationSystem.cpp b/Engine/Plugins/Runtime/AudioModulation/Source/AudioModulation/Private/AudioModulationSystem.cpp
--- a/Engine/Plugins/Runtime/AudioModulation/Source/AudioModulation/Private/AudioModulationSystem.cpp
+++ b/Engine/Plugins/Runtime/AudioModulation/Source/AudioModulation/Private/AudioModulationSystem.cpp
@@ -274,6 +274,21 @@ namespace AudioModulation
 		});
 	}
 
+	void FAudioModulationSystem::DeactivateBusMixByID(const uint32 ID)
+	{
+		ActiveBusMixIds.Remove(ID);
+
+		RunCommandOnProcessingThread([this, BusMixId = static_cast<FBusMixId>(ID)]()
+			{
+				FBusMixHandle MixHandle = FBusMixHandle::Get(BusMixId, RefProxies.BusMixes);
+				if (MixHandle.IsValid())
+				{
+					FModulatorBusMixProxy& MixProxy = MixHandle.FindProxy();
+					MixProxy.SetStopping();
+				}
+			});
+	}
+
 	void FAudioModulationSystem::DeactivateAllBusMixes()
 	{
 		ClearAllGlobalBusMixValues();
diff --git a/Engine/Plugins/Runtime/AudioModulation/Source/AudioModulation/Private/AudioModulationSystem.h b/Engine/Plugins/Runtime/AudioModulation/Source/AudioModulation/Private/AudioModulationSystem.h
--- a/Engine/Plugins/Runtime/AudioModulation/Source/AudioModulation/Private/AudioModulationSystem.h
+++ b/Engine/Plugins/Runtime/AudioModulation/Source/AudioModulation/Private/AudioModulationSystem.h
@@ -97,6 +97,7 @@ namespace AudioModulation
 
 		/** Deactivates given bus mix */
 		void DeactivateBusMix(const USoundControlBusMix& InBusMix);
+		void DeactivateBusMixByID(const uint32 ID);
 
 		/** Deactivates all bus mixes */
 		void DeactivateAllBusMixes();
@@ -337,6 +338,7 @@ namespace AudioModulation
 		void DeactivateAllBusMixes() { }
 		void DeactivateBus(const USoundControlBus& InBus) { }
 		void DeactivateBusMix(const USoundControlBusMix& InBusMix) { }
+		void DeactivateBusMixByID(const uint32 ID) {}
 		void DeactivateGenerator(const USoundModulationGenerator& InGenerator) { }
 
 		Audio::FDeviceId GetAudioDeviceId() const { return 0; }
diff --git a/Engine/Plugins/Runtime/AudioModulation/Source/AudioModulation/Private/MovieSceneAudioControlBusSection.cpp b/Engine/Plugins/Runtime/AudioModulation/Source/AudioModulation/Private/MovieSceneAudioControlBusSection.cpp
--- a/Engine/Plugins/Runtime/AudioModulation/Source/AudioModulation/Private/MovieSceneAudioControlBusSection.cpp
+++ b/Engine/Plugins/Runtime/AudioModulation/Source/AudioModulation/Private/MovieSceneAudioControlBusSection.cpp
@@ -15,7 +15,7 @@ UMovieSceneAudioControlBusSection::UMovieSceneAudioControlBusSection(const FObje
 	: Super(ObjectInitializer)
 {
 	BusValue.SetDefault(1.f);
-	MixBus = nullptr;
+	ControlBus = nullptr;
 
 	EvalOptions.EnableAndSetCompletionMode(EMovieSceneCompletionMode::ProjectDefault);
 }
@@ -27,7 +27,7 @@ EMovieSceneChannelProxyType UMovieSceneAudioControlBusSection::CacheChannelProxy
 
 #if WITH_EDITOR
 	FMovieSceneChannelMetaData Data;
-	Data.SetIdentifiers(TEXT("Control Value"), FText::FromName(MixBus == nullptr ? TEXT("Control Bus") : MixBus->GetFName()));
+	Data.SetIdentifiers(TEXT("Control Value"), FText::FromName(ControlBus == nullptr ? TEXT("Control Bus") : ControlBus->GetFName()));
 
 	Channels.Add(BusValue, Data, TMovieSceneExternalValue<float>::Make());
 #else
@@ -62,7 +62,7 @@ void UMovieSceneAudioControlBusSection::ImportEntityImpl(UMovieSceneEntitySystem
 	);
 }
 
-void UMovieSceneAudioControlBusSection::SetBus(USoundModulatorBase* NewBus)
+void UMovieSceneAudioControlBusSection::SetBus(USoundControlBus* NewBus)
 {
-	MixBus = NewBus;
+	ControlBus = NewBus;
 }
diff --git a/Engine/Plugins/Runtime/AudioModulation/Source/AudioModulation/Private/MovieSceneAudioControlBusSystem.cpp b/Engine/Plugins/Runtime/AudioModulation/Source/AudioModulation/Private/MovieSceneAudioControlBusSystem.cpp
--- a/Engine/Plugins/Runtime/AudioModulation/Source/AudioModulation/Private/MovieSceneAudioControlBusSystem.cpp
+++ b/Engine/Plugins/Runtime/AudioModulation/Source/AudioModulation/Private/MovieSceneAudioControlBusSystem.cpp
@@ -17,43 +17,11 @@
 #include "EntitySystem/MovieSceneEntitySystemRunner.h"
 #include "MovieSceneTracksComponentTypes.h"
 #include "MovieSceneAudioModulationComponentTypes.h"
+#include "SoundControlBus.h"
+#include "AudioModulationStatics.h"
 
 #include UE_INLINE_GENERATED_CPP_BY_NAME(MovieSceneAudioControlBusSystem)
 
-using AudioControlBusSectionData = UE::MovieScene::AudioControlBusSectionData;
-
-namespace UE::MovieScene
-{
-	enum class EPreAnimatedAudioControlBusStateType
-	{
-		ControlBusLifespan
-	};
-
-	template<typename BaseTraits>
-	struct FPreAnimatedAudioControlBusStateTraits : BaseTraits
-	{
-		using KeyType = FObjectKey;
-		using StorageType = EPreAnimatedAudioControlBusStateType;
-
-		EPreAnimatedAudioControlBusStateType CachePreAnimatedValue(FObjectKey InKey)
-		{
-			return EPreAnimatedAudioControlBusStateType::ControlBusLifespan;
-		}
-
-		void RestorePreAnimatedValue(FObjectKey InKey, EPreAnimatedAudioControlBusStateType InStateType, const FRestoreStateParams& Params)
-		{
-		}
-	};
-
-	using FPreAnimatedBoundObjectAudioControlBusStateTraits = FPreAnimatedAudioControlBusStateTraits<FBoundObjectPreAnimatedStateTraits>;
-
-	struct FPreAnimatedAudioControlBusStorage : TPreAnimatedStateStorage_ObjectTraits<FPreAnimatedBoundObjectAudioControlBusStateTraits>
-	{
-		static TAutoRegisterPreAnimatedStorageID<FPreAnimatedAudioControlBusStorage> StorageID;
-	};
-	TAutoRegisterPreAnimatedStorageID<FPreAnimatedAudioControlBusStorage> FPreAnimatedAudioControlBusStorage::StorageID;
-}
-
 namespace UE::MovieScene
 {
 	struct FEvaluateAudioControlBus
@@ -117,40 +85,54 @@ namespace UE::MovieScene
 			FObjectKey SectionKey(AudioSection);
 
 			// Grab the cached modulator
-			const AudioControlBusSectionData* EvaluationData = AudioSystem->GetAudioControlBusComponentEvaluationData(SectionKey);
+			const FAudioControlBusSectionData* EvaluationData = AudioSystem->GetAudioControlBusComponentEvaluationData(AudioSection);
 
-			check(AudioData.Section->MixBus);
+			if (AudioData.Section->ControlBus == nullptr)
+			{
+				return;
+			}
 
-			FString MixBusID = FString::Printf(TEXT("%s_%d"), *AudioData.Section->MixBus.GetFName().ToString(), Instance.GetSequenceID().GetInternalValue());
+			FString MixBusID = FString::Printf(TEXT("%s_%d"), *AudioData.Section->ControlBus.GetFName().ToString(), Instance.GetSequenceID().GetInternalValue());
 
 			const UObject* PlaybackContext = Instance.GetSharedPlaybackState()->GetPlaybackContext();
 			check(PlaybackContext);
-			const UWorld* World = PlaybackContext ? PlaybackContext->GetWorld() : nullptr;
+			UWorld* World = PlaybackContext ? PlaybackContext->GetWorld() : nullptr;
 			check(World);
 
-			if (FAudioDevice* AudioDevice = World->GetAudioDeviceRaw())
+			if (!UAudioModulationStatics::GetAudioWorld(World))
 			{
-				if (AudioDevice->IsModulationPluginEnabled())
+				return; // If we have no audio world then control bus mixes do not get created or updated properly
+			}
+
+			USoundControlBus* Bus = CastChecked<USoundControlBus>(AudioData.Section->ControlBus);
+
+			// If no cached value exists, we need to initialize
+			if (!EvaluationData || !EvaluationData->ControlBusMix->IsValidLowLevel() || EvaluationData->ControlBusMix == nullptr)
+			{
+				if (EvaluationData) // If we have evaluation data but the control mix asset is now invalid, we need to deactivate the unused mix before making a new one
+				{
+					UAudioModulationStatics::DeactivateBusMixByID(World, EvaluationData->BusMixID);
+
+					AudioSystem->RemoveAudioControlBusComponentEvaluationData(AudioSection);
+				}
+
+				if (USoundControlBusMix* NewMix = UAudioModulationStatics::CreateBusMixFromValue(World, Bus->GetFName(), {Bus}, BusValue))
 				{
-					if (IAudioModulationManager* Modulation = AudioDevice->ModulationInterface.Get())
-					{
-						// If no cached value exists, we need to initialize and put it in pre animated storage so it gets cleared correctly when sequence ends
-						if (!EvaluationData)
-						{
-							EvaluationData = AudioSystem->AddAudioControlBusComponentEvaluationData(InstanceHandle, AudioSection, AudioData.Section->MixBus, FName(*MixBusID));
-							
-							AudioSystem->PreAnimatedStorage->BeginTrackingEntity(EntityID, true, RootInstanceHandle, AudioData.Section->MixBus);
-							AudioSystem->PreAnimatedStorage->CachePreAnimatedValue(
-								FCachePreAnimatedValueParams(), AudioData.Section->MixBus,
-								[](FObjectKey InKey) { return EPreAnimatedAudioControlBusStateType::ControlBusLifespan; });
-							check(AudioDevice->ModulationInterface != nullptr);
-							AudioDevice->ModulationInterface->CreateNewSequencerBusMix(World, FName(*MixBusID), { AudioData.Section->MixBus }, BusValue);
-						}
-						check(EvaluationData);
-						Modulation->UpdateMixByModulator(World, FName(*MixBusID), *AudioData.Section->MixBus, BusValue);
-					}
+					EvaluationData = AudioSystem->AddAudioControlBusComponentEvaluationData(InstanceHandle, AudioSection, Bus, NewMix);
 				}
 			}
+
+			ensure(EvaluationData);
+
+			if (EvaluationData != nullptr)
+			{
+				if (!UAudioModulationStatics::IsControlBusMixActive(World, EvaluationData->ControlBusMix.Get()))
+				{
+					UAudioModulationStatics::ActivateBusMix(World, EvaluationData->ControlBusMix.Get());
+				}
+
+				UAudioModulationStatics::UpdateMixByFilter(World, EvaluationData->ControlBusMix.Get(), Bus->Address, nullptr, Bus->Parameter, BusValue, -1);
+			}		
 		}
 	};
 }
@@ -215,9 +197,6 @@ void UMovieSceneAudioControlBusSystem::OnRun(FSystemTaskPrerequisites& InPrerequ
 
 void UMovieSceneAudioControlBusSystem::OnLink()
 {
-	using namespace UE::MovieScene;
-
-	PreAnimatedStorage = Linker->PreAnimatedState.GetOrCreateStorage<FPreAnimatedAudioControlBusStorage>();
 }
 
 void UMovieSceneAudioControlBusSystem::OnUnlink()
@@ -229,37 +208,38 @@ void UMovieSceneAudioControlBusSystem::ClearBusMixes()
 {
 	if (FAudioDeviceManager* DeviceManager = FAudioDeviceManager::Get())
 	{
-		DeviceManager->IterateOverAllDevices([this](::Audio::FDeviceId DeviceId, FAudioDevice* AudioDevice)
+		if (FAudioDevice* AudioDevice = DeviceManager->GetAudioDeviceFromWorldContext(GetWorld()))
 		{
-			if (AudioDevice->IsModulationPluginEnabled())
+			for (const TPair<FObjectKey, FAudioControlBusSectionData> Modulator : ModulatorsByObjectKey)
 			{
-				if (IAudioModulationManager* ModulationManager = AudioDevice->ModulationInterface.Get())
-				{
-					for (const TPair<FObjectKey, AudioControlBusSectionData>& Modulator : ModulatorsByObjectKey)
-					{
-						ModulationManager->ClearSequencerBusMix(Modulator.Value.ControlBusMixName);
-					}
-				}
+				UAudioModulationStatics::DeactivateBusMix(GetWorld(), Modulator.Value.ControlBusMix.Get());
 			}
-		});
+		}
 	}
-
+	
 	ModulatorsByObjectKey.Reset();
 }
 
-const AudioControlBusSectionData* UMovieSceneAudioControlBusSystem::GetAudioControlBusComponentEvaluationData(const FObjectKey SectionKey) const
+const FAudioControlBusSectionData* UMovieSceneAudioControlBusSystem::GetAudioControlBusComponentEvaluationData(const TObjectPtr<UMovieSceneAudioControlBusSection> SectionKey) const
 {
-	const AudioControlBusSectionData* data = ModulatorsByObjectKey.Find(SectionKey);
+	const FAudioControlBusSectionData* data = ModulatorsByObjectKey.Find(SectionKey);
 
 	return data;
 }
 
-AudioControlBusSectionData* UMovieSceneAudioControlBusSystem::AddAudioControlBusComponentEvaluationData(FInstanceHandle InstanceHandle, FObjectKey SectionKey, USoundModulatorBase* SectionData, FName ID)
+FAudioControlBusSectionData* UMovieSceneAudioControlBusSystem::AddAudioControlBusComponentEvaluationData(FInstanceHandle InstanceHandle, TObjectPtr<UMovieSceneAudioControlBusSection> SectionKey, USoundModulatorBase* SectionData, USoundControlBusMix* Mix)
 {
 	using namespace UE::MovieScene;
 	check(SectionData != nullptr);
 
-	AudioControlBusSectionData NewData = AudioControlBusSectionData(SectionData, ID);
+	FAudioControlBusSectionData NewData = FAudioControlBusSectionData(Mix, Mix->GetUniqueID());
+
+	FAudioControlBusSectionData* Result = &ModulatorsByObjectKey.Add(SectionKey, NewData);
 
-	return &ModulatorsByObjectKey.Add(SectionKey, NewData);
+	return Result;
+}
+
+void UMovieSceneAudioControlBusSystem::RemoveAudioControlBusComponentEvaluationData(TObjectPtr<UMovieSceneAudioControlBusSection> SectionKey)
+{
+	ModulatorsByObjectKey.Remove(SectionKey);
 }
diff --git a/Engine/Plugins/Runtime/AudioModulation/Source/AudioModulation/Private/MovieSceneAudioControlBusSystem.h b/Engine/Plugins/Runtime/AudioModulation/Source/AudioModulation/Private/MovieSceneAudioControlBusSystem.h
--- a/Engine/Plugins/Runtime/AudioModulation/Source/AudioModulation/Private/MovieSceneAudioControlBusSystem.h
+++ b/Engine/Plugins/Runtime/AudioModulation/Source/AudioModulation/Private/MovieSceneAudioControlBusSystem.h
@@ -1,26 +1,29 @@
 // Copyright Epic Games, Inc. All Rights Reserved.
 
 #pragma once
+
+#include "AudioModulation.h"
 #include "CoreTypes.h"
 #include "EntitySystem/MovieSceneEntityIDs.h"
 #include "EntitySystem/MovieSceneEntitySystem.h"
 #include "EntitySystem/MovieSceneSequenceInstanceHandle.h"
+#include "MovieSceneAudioControlBusSection.h"
 #include "MovieSceneTracksComponentTypes.h"
 #include "UObject/ObjectMacros.h"
 #include "IAudioModulation.h"
+#include "SoundControlBusMix.h"
 
 #include "MovieSceneAudioControlBusSystem.generated.h"
 
-namespace UE::MovieScene
+USTRUCT()
+struct FAudioControlBusSectionData
 {
-	struct FPreAnimatedAudioControlBusStorage;
+	GENERATED_BODY()
 
-	struct AudioControlBusSectionData
-	{
-		TObjectPtr<USoundModulatorBase> ControlBus;
-		FName ControlBusMixName;
-	};
-}
+	UPROPERTY(Transient)
+	TObjectPtr<USoundControlBusMix> ControlBusMix;
+	uint32 BusMixID;
+};
 
 UCLASS()
 class UMovieSceneAudioControlBusSystem : public UMovieSceneEntitySystem
@@ -29,6 +32,7 @@ class UMovieSceneAudioControlBusSystem : public UMovieSceneEntitySystem
 
 public:
 	using FInstanceHandle = UE::MovieScene::FInstanceHandle;
+
 	UMovieSceneAudioControlBusSystem(const FObjectInitializer& ObjInit);
 
 	virtual void OnSchedulePersistentTasks(UE::MovieScene::IEntitySystemScheduler* TaskScheduler) override;
@@ -39,15 +43,11 @@ public:
 
 	void ClearBusMixes();
 
-	using AudioControlBusSectionData = UE::MovieScene::AudioControlBusSectionData;
-
-	const AudioControlBusSectionData* GetAudioControlBusComponentEvaluationData(const FObjectKey SectionKey) const;
-	AudioControlBusSectionData* AddAudioControlBusComponentEvaluationData(FInstanceHandle InstanceHandle, FObjectKey SectionKey, USoundModulatorBase* SectionData, FName ID);
-
-	TArray<USoundModulatorBase*> Modulators;
-
-	/** Pre-animated state */
-	TSharedPtr<UE::MovieScene::FPreAnimatedAudioControlBusStorage> PreAnimatedStorage;
+	const FAudioControlBusSectionData* GetAudioControlBusComponentEvaluationData(const TObjectPtr<UMovieSceneAudioControlBusSection> SectionKey) const;
+	FAudioControlBusSectionData* AddAudioControlBusComponentEvaluationData(FInstanceHandle InstanceHandle, TObjectPtr<UMovieSceneAudioControlBusSection> SectionKey, USoundModulatorBase* SectionData, USoundControlBusMix* Mix);
+	void RemoveAudioControlBusComponentEvaluationData(TObjectPtr<UMovieSceneAudioControlBusSection> SectionKey);
 
-	TMap<FObjectKey, AudioControlBusSectionData> ModulatorsByObjectKey;
+protected:
+	UPROPERTY(Transient)
+	TMap<TObjectPtr<UMovieSceneAudioControlBusSection>, FAudioControlBusSectionData> ModulatorsByObjectKey;
 };
\ No newline at end of file
diff --git a/Engine/Plugins/Runtime/AudioModulation/Source/AudioModulation/Private/MovieSceneAudioControlBusTrack.cpp b/Engine/Plugins/Runtime/AudioModulation/Source/AudioModulation/Private/MovieSceneAudioControlBusTrack.cpp
--- a/Engine/Plugins/Runtime/AudioModulation/Source/AudioModulation/Private/MovieSceneAudioControlBusTrack.cpp
+++ b/Engine/Plugins/Runtime/AudioModulation/Source/AudioModulation/Private/MovieSceneAudioControlBusTrack.cpp
@@ -13,7 +13,7 @@ UMovieSceneAudioControlBusTrack::UMovieSceneAudioControlBusTrack(const FObjectIn
 {
 }
 
-UMovieSceneSection* UMovieSceneAudioControlBusTrack::AddNewControlBus(USoundModulatorBase* Bus)
+UMovieSceneSection* UMovieSceneAudioControlBusTrack::AddNewControlBus(USoundControlBus* Bus)
 {
 	check(Bus);
 
@@ -90,9 +90,18 @@ void UMovieSceneAudioControlBusTrack::GetCompatibleUserDecorationsImpl(TSet<UCla
 #if WITH_EDITOR
 FText UMovieSceneAudioControlBusTrack::GetDefaultDisplayName() const
 {
+	ensure(AudioSections.Num() > 0);
+
 	if (AudioSections.Num() > 0)
 	{
-		return FText::FromString(Cast<UMovieSceneAudioControlBusSection>(AudioSections[0])->MixBus->GetName());
+		const UMovieSceneAudioControlBusSection* Section = Cast<UMovieSceneAudioControlBusSection>(AudioSections[0]);
+
+		check(Section);
+
+		if (Section->ControlBus)
+		{
+			return FText::FromString(Section->ControlBus->GetName());
+		}
 	}
 
 	return NSLOCTEXT("MovieSceneAudioControlBusTrack","DefaultTrackName", "Control Bus");
@@ -100,13 +109,22 @@ FText UMovieSceneAudioControlBusTrack::GetDefaultDisplayName() const
 
 void UMovieSceneAudioControlBusTrack::PostLoad()
 {
-	if (AudioSections.Num() > 0 && GetDisplayName().EqualTo(FText::FromString(Cast<UMovieSceneAudioControlBusSection>(AudioSections[0])->MixBus->GetName())))
+	ensure(AudioSections.Num() > 0);
+
+	if (AudioSections.Num() > 0)
 	{
-		UMovieSceneSequence* Sequence = GetTypedOuter<UMovieSceneSequence>();
-		UE_LOG(LogMovieScene, Display, TEXT("Control Bus Track: Removing display name %s from %s because it matches default asset name and will not update automatically with asset changes"), *AudioSections[0]->GetName(), *GetNameSafe(Sequence));
-		SetDisplayName(FText::GetEmpty());
+		const UMovieSceneAudioControlBusSection* Section = Cast<UMovieSceneAudioControlBusSection>(AudioSections[0]);
+
+		check(Section);
+
+		if (Section->ControlBus && GetDisplayName().EqualTo(FText::FromString(Section->ControlBus->GetName())))
+		{
+			UMovieSceneSequence* Sequence = GetTypedOuter<UMovieSceneSequence>();
+			UE_LOG(LogMovieScene, Display, TEXT("Control Bus Track: Removing display name %s from %s because it matches default asset name and will not update automatically with asset changes"), *AudioSections[0]->GetName(), *GetNameSafe(Sequence));
+			SetDisplayName(FText::GetEmpty());
+		}
 	}
 	Super::PostLoad();
 }
 
-#endif
+#endif // WITH_EDITOR
diff --git a/Engine/Plugins/Runtime/AudioModulation/Source/AudioModulation/Private/SoundControlBusMix.cpp b/Engine/Plugins/Runtime/AudioModulation/Source/AudioModulation/Private/SoundControlBusMix.cpp
--- a/Engine/Plugins/Runtime/AudioModulation/Source/AudioModulation/Private/SoundControlBusMix.cpp
+++ b/Engine/Plugins/Runtime/AudioModulation/Source/AudioModulation/Private/SoundControlBusMix.cpp
@@ -40,21 +40,20 @@ void USoundControlBusMix::BeginDestroy()
 {
 	using namespace AudioModulation;
 
-	if (UWorld* World = GetWorld())
+	if (FAudioDeviceManager* DeviceManager = FAudioDeviceManager::Get())
 	{
-		FAudioDeviceHandle AudioDevice = World->GetAudioDevice();
-		if (AudioDevice.IsValid())
-		{
-			if (AudioDevice->IsModulationPluginEnabled())
+		DeviceManager->IterateOverAllDevices([this](::Audio::FDeviceId DeviceId, FAudioDevice* AudioDevice)
 			{
-				if (IAudioModulationManager* ModulationInterface = AudioDevice->ModulationInterface.Get())
+				if (AudioDevice->IsModulationPluginEnabled())
 				{
-					FAudioModulationManager* Modulation = static_cast<FAudioModulationManager*>(ModulationInterface);
-					check(Modulation);
-					Modulation->DeactivateBusMix(*this);
+					if (IAudioModulationManager* ModulationInterface = AudioDevice->ModulationInterface.Get())
+					{
+						FAudioModulationManager* Modulation = static_cast<FAudioModulationManager*>(ModulationInterface);
+						check(Modulation);
+						Modulation->DeactivateBusMix(*this);
+					}
 				}
-			}
-		}
+			});
 	}
 
 	// Call parent destroy at end to ensure object is in a valid state for the modulation manager to clean up first
diff --git a/Engine/Plugins/Runtime/AudioModulation/Source/AudioModulation/Public/AudioModulation.h b/Engine/Plugins/Runtime/AudioModulation/Source/AudioModulation/Public/AudioModulation.h
--- a/Engine/Plugins/Runtime/AudioModulation/Source/AudioModulation/Public/AudioModulation.h
+++ b/Engine/Plugins/Runtime/AudioModulation/Source/AudioModulation/Public/AudioModulation.h
@@ -40,24 +40,6 @@ namespace AudioModulation
 
 		UE_API virtual void GetControlBusAssetPath(TArray<FTopLevelAssetPath>& OutClassNames) override;
 
-		UE_API virtual void CreateNewSequencerBusMix(
-			const UObject* WorldContextObject,
-			FName Name,
-			const TArray<USoundModulatorBase*>& Buses,
-			float Value = 1.0f,
-			float AttackTime = 0.1f,
-			float ReleaseTime = 0.1f,
-			bool bActivate = true) override;
-
-		UE_API virtual void UpdateMixByModulator(
-			const UObject* WorldContextObject,
-			FName Mix,
-			USoundModulatorBase& Modulator,
-			float Value = 1.0f,
-			float FadeTime = -1.0f);
-
-		UE_API virtual void ClearSequencerBusMix(FName Name) override;
-
 #if !UE_BUILD_SHIPPING
 		UE_API virtual void SetDebugBusFilter(const FString* InFilter);
 		UE_API virtual void SetDebugMixFilter(const FString* InFilter);
@@ -98,6 +80,7 @@ namespace AudioModulation
 		UE_API void DeactivateBus(const USoundControlBus& InBus);
 
 		UE_API void DeactivateBusMix(const USoundControlBusMix& InBusMix);
+		UE_API void DeactivateBusMixByID(const uint32 ID);
 		UE_API void DeactivateAllBusMixes();
 
 		UE_DEPRECATED(5.4, "Deactivation of modulators in this manner is now deprecated. Use USoundModulationDestination to safety activate and track a given modulator")
@@ -124,7 +107,6 @@ namespace AudioModulation
 
 	private:
 		FAudioModulationSystem* ModSystem = nullptr;
-		TMap<FName, TObjectPtr<USoundControlBusMix>> SequencerBusMixes;
 	};
 
 	AUDIOMODULATION_API FAudioModulationManager* GetDeviceModulationManager(Audio::FDeviceId InDeviceId);
diff --git a/Engine/Plugins/Runtime/AudioModulation/Source/AudioModulation/Public/AudioModulationStatics.h b/Engine/Plugins/Runtime/AudioModulation/Source/AudioModulation/Public/AudioModulationStatics.h
--- a/Engine/Plugins/Runtime/AudioModulation/Source/AudioModulation/Public/AudioModulationStatics.h
+++ b/Engine/Plugins/Runtime/AudioModulation/Source/AudioModulation/Public/AudioModulationStatics.h
@@ -240,6 +240,11 @@ public:
 	)
 	static UE_API void DeactivateBusMix(const UObject* WorldContextObject, USoundControlBusMix* Mix);
 
+	///** Deactivates a modulation bus mix. Does nothing if an instance of the provided bus mix is already inactive.
+	// * @param MixID - Mix to deactivate
+	// */
+	static UE_API void DeactivateBusMixByID(const UObject* WorldContextObject, uint32 MixID);
+
 	/** Deactivates a modulation generator. Does nothing if an instance of the provided generator is already inactive.
 	 * @param Generator - Generator to activate
 	 */
diff --git a/Engine/Plugins/Runtime/AudioModulation/Source/AudioModulation/Public/MovieSceneAudioControlBusSection.h b/Engine/Plugins/Runtime/AudioModulation/Source/AudioModulation/Public/MovieSceneAudioControlBusSection.h
--- a/Engine/Plugins/Runtime/AudioModulation/Source/AudioModulation/Public/MovieSceneAudioControlBusSection.h
+++ b/Engine/Plugins/Runtime/AudioModulation/Source/AudioModulation/Public/MovieSceneAudioControlBusSection.h
@@ -6,10 +6,12 @@
 #include "EntitySystem/IMovieSceneEntityProvider.h"
 #include "IAudioModulation.h"
 #include "MovieSceneSection.h"
+#include "SoundControlBus.h"
 #include "UObject/ObjectMacros.h"
 
 #include "MovieSceneAudioControlBusSection.generated.h"
 
+
 UCLASS(MinimalAPI)
 class UMovieSceneAudioControlBusSection 
 	: public UMovieSceneSection
@@ -23,12 +25,12 @@ public:
 	AUDIOMODULATION_API virtual void ImportEntityImpl(UMovieSceneEntitySystemLinker* EntityLinker, const FEntityImportParams& Params, FImportedEntity* OutImportedEntity) override;
 
 
-	AUDIOMODULATION_API void SetBus(USoundModulatorBase* NewBus);
+	AUDIOMODULATION_API void SetBus(USoundControlBus* NewBus);
 
 	UPROPERTY()
 	FMovieSceneFloatChannel BusValue;
 
 	UPROPERTY()
-	TObjectPtr<USoundModulatorBase> MixBus;
+	TObjectPtr<USoundControlBus> ControlBus;
 };
 
diff --git a/Engine/Plugins/Runtime/AudioModulation/Source/AudioModulation/Public/MovieSceneAudioControlBusTrack.h b/Engine/Plugins/Runtime/AudioModulation/Source/AudioModulation/Public/MovieSceneAudioControlBusTrack.h
--- a/Engine/Plugins/Runtime/AudioModulation/Source/AudioModulation/Public/MovieSceneAudioControlBusTrack.h
+++ b/Engine/Plugins/Runtime/AudioModulation/Source/AudioModulation/Public/MovieSceneAudioControlBusTrack.h
@@ -5,6 +5,7 @@
 #include "IAudioModulation.h"
 #include "MovieSceneNameableTrack.h"
 #include "MovieSceneSection.h"
+#include "SoundControlBus.h"
 #include "UObject/ObjectMacros.h"
 
 #include "MovieSceneAudioControlBusTrack.generated.h"
@@ -17,7 +18,7 @@ class UMovieSceneAudioControlBusTrack
 
 public:
 
-	AUDIOMODULATION_API UMovieSceneSection* AddNewControlBus(USoundModulatorBase* Bus);
+	AUDIOMODULATION_API UMovieSceneSection* AddNewControlBus(USoundControlBus* Bus);
 
 	// UMovieSceneTrack interface
 
diff --git a/Engine/Source/Runtime/AudioExtensions/Public/IAudioModulation.h b/Engine/Source/Runtime/AudioExtensions/Public/IAudioModulation.h
--- a/Engine/Source/Runtime/AudioExtensions/Public/IAudioModulation.h
+++ b/Engine/Source/Runtime/AudioExtensions/Public/IAudioModulation.h
@@ -202,24 +202,6 @@ public:
 
 	virtual void GetControlBusAssetPath(TArray<FTopLevelAssetPath>& OutClassNames) = 0;
 
-	virtual void CreateNewSequencerBusMix(
-		const UObject* WorldContextObject,
-		FName Name,
-		const TArray<USoundModulatorBase*>& Buses,
-		float Value = 1.0f,
-		float AttackTime = 0.1f,
-		float ReleaseTime = 0.1f,
-		bool bActivate = true) = 0;
-
-	virtual void UpdateMixByModulator(
-		const UObject* WorldContextObject,
-		FName Mix,
-		USoundModulatorBase& Modulator,
-		float Value = 1.0f,
-		float FadeTime = -1.0f) = 0;
-
-	virtual void ClearSequencerBusMix(FName Name) = 0;
-
 #if !UE_BUILD_SHIPPING
 	/** Request to post help from active plugin (non-shipping builds only) */
 	virtual bool OnPostHelp(FCommonViewportClient* ViewportClient, const TCHAR* Stream) = 0; 
