[AudioInsights] - Modulation Refresh (Pt. 1/2): Backend

New features:
* Supports Parameter Patches
* Supports Generators
* Supports IsBypassed values
* Supports modulators being added on the fly
* Supports Contributing Modulators
-> E.g. Control Bus can have contributing modualtors: Generators / Control Bus Mixes
-> E.g. Parameter Patches can have contributing modulators: Control Buses mixed with a specified curve

#Synced-CL 50595586

diff --git a/Engine/Plugins/Runtime/AudioModulation/Source/AudioModulation/AudioModulation.Build.cs b/Engine/Plugins/Runtime/AudioModulation/Source/AudioModulation/AudioModulation.Build.cs
--- a/Engine/Plugins/Runtime/AudioModulation/Source/AudioModulation/AudioModulation.Build.cs
+++ b/Engine/Plugins/Runtime/AudioModulation/Source/AudioModulation/AudioModulation.Build.cs
@@ -32,6 +32,7 @@ namespace UnrealBuildTool.Rules
 					"CoreUObject",
 					"DeveloperSettings",
 					"Engine",
+					"Projects",
 					"SignalProcessing",
 					"TraceLog"
 				}
diff --git a/Engine/Plugins/Runtime/AudioModulation/Source/AudioModulation/Private/AudioModulationSystem.cpp b/Engine/Plugins/Runtime/AudioModulation/Source/AudioModulation/Private/AudioModulationSystem.cpp
--- a/Engine/Plugins/Runtime/AudioModulation/Source/AudioModulation/Private/AudioModulationSystem.cpp
+++ b/Engine/Plugins/Runtime/AudioModulation/Source/AudioModulation/Private/AudioModulationSystem.cpp
@@ -1,6 +1,7 @@
 // Copyright Epic Games, Inc. All Rights Reserved.
 #include "AudioModulationSystem.h"
 #include "UObject/Package.h"
+#include <type_traits>
 
 #if WITH_AUDIOMODULATION
 
@@ -20,6 +21,7 @@
 #include "Generators/SoundModulationLFO.h"
 #include "HAL/PlatformTLS.h"
 #include "IAudioModulation.h"
+#include "Interfaces/IPluginManager.h"
 #include "Misc/CoreDelegates.h"
 #include "ProfilingDebugging/CpuProfilerTrace.h"
 #include "SoundControlBusProxy.h"
@@ -37,61 +39,36 @@ DECLARE_DWORD_COUNTER_STAT(TEXT("Patch Count"), STAT_AudioModulationPatchCount,
 DECLARE_DWORD_COUNTER_STAT(TEXT("Render Queue Commands Processed"), STAT_AudioModulationProcQueueCount, STATGROUP_AudioModulation)
 
 #if UE_AUDIO_PROFILERTRACE_ENABLED
-UE_TRACE_EVENT_BEGIN(Audio, ControlBusUpdate)
-	UE_TRACE_EVENT_FIELD(uint32, DeviceId)
-	UE_TRACE_EVENT_FIELD(uint32, ControlBusId)
-	UE_TRACE_EVENT_FIELD(uint64, Timestamp)
-	UE_TRACE_EVENT_FIELD(float, Value)
-UE_TRACE_EVENT_END()
-
-UE_TRACE_EVENT_BEGIN(Audio, ControlBusActivePing)
-	UE_TRACE_EVENT_FIELD(uint32, DeviceId)
-	UE_TRACE_EVENT_FIELD(uint32, ControlBusId)
-	UE_TRACE_EVENT_FIELD(uint64, Timestamp)
-	UE_TRACE_EVENT_FIELD(UE::Trace::WideString, Name)
-	UE_TRACE_EVENT_FIELD(UE::Trace::WideString, ParamName)
-UE_TRACE_EVENT_END()
 
-UE_TRACE_EVENT_BEGIN(Audio, BusMixRegisterBus)
-	UE_TRACE_EVENT_FIELD(uint32, DeviceId)
-	UE_TRACE_EVENT_FIELD(uint32, SourceId)
-	UE_TRACE_EVENT_FIELD(uint64, Timestamp)
-	UE_TRACE_EVENT_FIELD(uint32, ModulatingSourceId)
-	UE_TRACE_EVENT_FIELD(UE::Trace::WideString, BusName)
-UE_TRACE_EVENT_END()
-
-UE_TRACE_EVENT_BEGIN(Audio, BusMixActivate)
-	UE_TRACE_EVENT_FIELD(uint32, DeviceId)
-	UE_TRACE_EVENT_FIELD(uint32, SourceId)
-	UE_TRACE_EVENT_FIELD(uint64, Timestamp)
+// Mirrors FModulatorTraceMessageBase: Engine/Plugins/Runtime/AudioModulationInsights/Source/AudioModulationInsights/Private/Messages/AudioModulationTraceMessages.h
+#define MODULATOR_TRACE_MESSAGE_BASE \
+	UE_TRACE_EVENT_FIELD(uint32, DeviceId) \
+	UE_TRACE_EVENT_FIELD(uint64, Timestamp) \
+	UE_TRACE_EVENT_FIELD(Audio::FModulatorId, ModulatorId)
+
+// Mirrors FActivateModulatorTraceMessage: Engine/Plugins/Runtime/AudioModulationInsights/Source/AudioModulationInsights/Private/Messages/AudioModulationTraceMessages.h
+// Note: This data doesn't change during the lifetime of the Modulator
+UE_TRACE_EVENT_BEGIN(Audio, ActivateModulatorTraceMessage)
+	MODULATOR_TRACE_MESSAGE_BASE
+	UE_TRACE_EVENT_FIELD(Audio::FModulatorTypeId, ModulatorType)
 	UE_TRACE_EVENT_FIELD(UE::Trace::WideString, Name)
 UE_TRACE_EVENT_END()
 
-UE_TRACE_EVENT_BEGIN(Audio, BusMixUpdate)
-	UE_TRACE_EVENT_FIELD(uint32, DeviceId)
-	UE_TRACE_EVENT_FIELD(uint32, SourceId)
-	UE_TRACE_EVENT_FIELD(uint64, Timestamp)
-	UE_TRACE_EVENT_FIELD(uint32[], BusIds)
-	UE_TRACE_EVENT_FIELD(float[], BusValues)
-UE_TRACE_EVENT_END()
-
-UE_TRACE_EVENT_BEGIN(Audio, GeneratorUpdate)
-	UE_TRACE_EVENT_FIELD(uint32, DeviceId)
-	UE_TRACE_EVENT_FIELD(uint32, SourceId)
-	UE_TRACE_EVENT_FIELD(uint64, Timestamp)
-	UE_TRACE_EVENT_FIELD(uint32[], BusIds)
-	UE_TRACE_EVENT_FIELD(float, GeneratorValue)
+// Mirrors FUpdateModulatorTraceMessage: Engine/Plugins/Runtime/AudioModulationInsights/Source/AudioModulationInsights/Private/Messages/AudioModulationTraceMessages.h
+// Note: This data can change during the lifetime of the Modulator
+UE_TRACE_EVENT_BEGIN(Audio, UpdateModulatorTraceMessage)
+	MODULATOR_TRACE_MESSAGE_BASE
+	UE_TRACE_EVENT_FIELD(float, ModulatorValue)
+	UE_TRACE_EVENT_FIELD(bool, IsBypassed)
+	UE_TRACE_EVENT_FIELD(Audio::FModulatorId[], ContributingModulatorIds)
+	UE_TRACE_EVENT_FIELD(float[], ContributingModulatorValues)
 UE_TRACE_EVENT_END()
 
-UE_TRACE_EVENT_BEGIN(Audio, BusFinalValuesUpdate)
-	UE_TRACE_EVENT_FIELD(uint32, DeviceId)
-	UE_TRACE_EVENT_FIELD(uint32, SourceId)
-	UE_TRACE_EVENT_FIELD(uint64, Timestamp)
-	UE_TRACE_EVENT_FIELD(uint32[], BusIds)
-	UE_TRACE_EVENT_FIELD(float[], BusValues)
+// Mirrors FDeactivateModulatorTraceMessage: Engine/Plugins/Runtime/AudioModulationInsights/Source/AudioModulationInsights/Private/Messages/AudioModulationTraceMessages.h
+UE_TRACE_EVENT_BEGIN(Audio, DeactivateModulatorTraceMessage)
+	MODULATOR_TRACE_MESSAGE_BASE
 UE_TRACE_EVENT_END()
 
-UE_TRACE_EVENT_DEFINE(Audio, ModulatingSourceDeactivate)
 #endif // UE_AUDIO_PROFILERTRACE_ENABLED
 
 namespace AudioModulation
@@ -120,6 +97,10 @@ namespace AudioModulation
 
 #if UE_AUDIO_PROFILERTRACE_ENABLED
 		FTraceAuxiliary::OnTraceStarted.AddRaw(this, &FAudioModulationSystem::OnTraceStarted);
+
+		// Note: Enabling/Disabling plugins require restarting the editor, so we only need to determine a plugins state once here within Initialize
+		const TSharedPtr<const IPlugin> AudioModulationInsightsPlugin = IPluginManager::Get().FindPlugin("AudioModulationInsights");
+		bIsAudioModulationInsightsPluginEnabled = (AudioModulationInsightsPlugin.IsValid() && AudioModulationInsightsPlugin->IsEnabled());
 #endif // UE_AUDIO_PROFILERTRACE_ENABLED
 
 		AudioDeviceId = InitializationParams.AudioDevicePtr->DeviceID;
@@ -176,24 +157,6 @@ namespace AudioModulation
 			}
 
 			ManuallyActivatedBusMixes.Add(MoveTemp(BusMixHandle));
-
-#if UE_AUDIO_PROFILERTRACE_ENABLED
-			for (const FModulatorBusMixStageSettings& StageSetting : Settings.Stages)
-			{
-				UE_TRACE_LOG(Audio, BusMixRegisterBus, AudioChannel)
-					<< BusMixRegisterBus.DeviceId(AudioDeviceId)
-					<< BusMixRegisterBus.SourceId(StageSetting.BusSettings.GetId())
-					<< BusMixRegisterBus.Timestamp(FPlatformTime::Cycles64())
-					<< BusMixRegisterBus.ModulatingSourceId(Settings.GetId())
-					<< BusMixRegisterBus.BusName(*(StageSetting.BusSettings.GetPath().ToString()));
-			}
-
-			UE_TRACE_LOG(Audio, BusMixActivate, AudioChannel)
-				<< BusMixActivate.DeviceId(AudioDeviceId)
-				<< BusMixActivate.SourceId(Settings.GetId())
-				<< BusMixActivate.Timestamp(FPlatformTime::Cycles64())
-				<< BusMixActivate.Name(*Settings.GetPath().ToString());
-#endif
 		});
 	}
 
@@ -300,13 +263,6 @@ namespace AudioModulation
 			for (auto& [BusMixId, MixProxy] : RefProxies.BusMixes)
 			{
 				MixProxy.SetStopping();
-
-#if UE_AUDIO_PROFILERTRACE_ENABLED
-				UE_TRACE_LOG(Audio, ModulatingSourceDeactivate, AudioChannel)
-					<< ModulatingSourceDeactivate.DeviceId(AudioDeviceId)
-					<< ModulatingSourceDeactivate.SourceId(BusMixId)
-					<< ModulatingSourceDeactivate.Timestamp(FPlatformTime::Cycles64());
-#endif
 			}
 		});
 	}
@@ -674,101 +630,7 @@ namespace AudioModulation
 		SET_DWORD_STAT(STAT_AudioModulationProcQueueCount, CommandsProcessed);
 
 #if UE_AUDIO_PROFILERTRACE_ENABLED
-		const bool bChannelEnabled = UE_TRACE_CHANNELEXPR_IS_ENABLED(AudioChannel);
-		if (bChannelEnabled)
-		{
-			TMap<FGeneratorId, TSet<FBusId>> GeneratorIdToBusIdsMap;
-			TArray<FBusId, TInlineAllocator<32>> BusIdsFinalValues;
-			TArray<float, TInlineAllocator<32>> BusFinalValues;
-
-			const int32 NumBuses = RefProxies.Buses.Num();
-			BusIdsFinalValues.Reserve(NumBuses);
-			BusFinalValues.Reserve(NumBuses);
-
-			ControlBusActivePingTimeSec -= InElapsed;
-			const bool bSendControlBusActivePingTrace = (ControlBusActivePingTimeSec <= 0.0f);
-			if (bSendControlBusActivePingTrace)
-			{
-				ControlBusActivePingTimeSec = ControlBusActivePingIntervalTimeSec;
-			}
-
-			// Control Bus message trace
-			for (const auto& [BusId, ControlBusProxy] : RefProxies.Buses)
-			{
-				if (bSendControlBusActivePingTrace)
-				{
-					UE_TRACE_LOG(Audio, ControlBusActivePing, AudioChannel)
-						<< ControlBusActivePing.DeviceId(static_cast<uint32>(AudioDeviceId))
-						<< ControlBusActivePing.ControlBusId(static_cast<uint32>(BusId))
-						<< ControlBusActivePing.Timestamp(FPlatformTime::Cycles64())
-						<< ControlBusActivePing.Name(*(ControlBusProxy.GetPath().ToString()))
-						<< ControlBusActivePing.ParamName(*(ControlBusProxy.GetParameterName().ToString()));
-				}
-
-				UE_TRACE_LOG(Audio, ControlBusUpdate, AudioChannel)
-					<< ControlBusUpdate.DeviceId(static_cast<uint32>(AudioDeviceId))
-					<< ControlBusUpdate.ControlBusId(static_cast<uint32>(BusId))
-					<< ControlBusUpdate.Timestamp(FPlatformTime::Cycles64())
-					<< ControlBusUpdate.Value(ControlBusProxy.GetValue());
-
-				// Collect the final value of this control bus
-				BusIdsFinalValues.Emplace(BusId);
-				BusFinalValues.Emplace(ControlBusProxy.GetValue());
-
-				// Collect generators set in this control bus
-				for (const FGeneratorHandle& GeneratorHandle : ControlBusProxy.GetGeneratorHandles())
-				{
-					GeneratorIdToBusIdsMap.FindOrAdd(GeneratorHandle.GetId()).Emplace(BusId);
-				}
-			}
-
-			// Control Bus Mix message trace
-			for (const auto& [BusMixId, BusMixProxy] : RefProxies.BusMixes)
-			{
-				const int32 NumStages = BusMixProxy.Stages.Num();
-
-				TArray<FBusId, TInlineAllocator<32>> BusIds;
-				TArray<float,  TInlineAllocator<32>> BusMixCurrentValues;
-
-				BusIds.Reserve(NumStages);
-				BusMixCurrentValues.Reserve(NumStages);
-
-				for (const auto& [BusId, BusMixStageProxy] : BusMixProxy.Stages)
-				{
-					BusIds.Emplace(BusId);
-					BusMixCurrentValues.Emplace(BusMixStageProxy.Value.GetCurrentValue());
-				}
-
-				UE_TRACE_LOG(Audio, BusMixUpdate, AudioChannel)
-					<< BusMixUpdate.DeviceId(AudioDeviceId)
-					<< BusMixUpdate.SourceId(BusMixId)
-					<< BusMixUpdate.Timestamp(FPlatformTime::Cycles64())
-					<< BusMixUpdate.BusIds(BusIds.GetData(), BusIds.Num())
-					<< BusMixUpdate.BusValues(BusMixCurrentValues.GetData(), BusMixCurrentValues.Num());
-			}
-
-			// Generator message trace
-			for (const auto& [GeneratorId, BusIds] : GeneratorIdToBusIdsMap)
-			{
-				if (const FModulatorGeneratorProxy* FoundGeneratorProxyPtr = RefProxies.Generators.Find(GeneratorId))
-				{
-					UE_TRACE_LOG(Audio, GeneratorUpdate, AudioChannel)
-						<< GeneratorUpdate.DeviceId(AudioDeviceId)
-						<< GeneratorUpdate.SourceId(GeneratorId)
-						<< GeneratorUpdate.Timestamp(FPlatformTime::Cycles64())
-						<< GeneratorUpdate.BusIds(BusIds.Array().GetData(), BusIds.Num())
-						<< GeneratorUpdate.GeneratorValue(FoundGeneratorProxyPtr->GetValue());
-				}
-			}
-
-			// Bus Final values message trace
-			UE_TRACE_LOG(Audio, BusFinalValuesUpdate, AudioChannel)
-				<< BusFinalValuesUpdate.DeviceId(AudioDeviceId)
-				<< BusFinalValuesUpdate.SourceId(INDEX_NONE)
-				<< BusFinalValuesUpdate.Timestamp(FPlatformTime::Cycles64())
-				<< BusFinalValuesUpdate.BusIds(BusIdsFinalValues.GetData(), BusIdsFinalValues.Num())
-				<< BusFinalValuesUpdate.BusValues(BusFinalValues.GetData(), BusFinalValues.Num());
-		}
+		UpdateAudioInsightsTraceMessages(InElapsed);
 #endif // UE_AUDIO_PROFILERTRACE_ENABLED
 
 #if !UE_BUILD_SHIPPING
@@ -879,30 +741,187 @@ namespace AudioModulation
 #if UE_AUDIO_PROFILERTRACE_ENABLED
 	void FAudioModulationSystem::OnTraceStarted(FTraceAuxiliary::EConnectionType TraceType, const FString& TraceDestination)
 	{
-		RunCommandOnProcessingThread([this]()
+		// Ensures when a trace is started, our active trace messages are sent instantly without delay
+		ModulatorActivePingTimeSec = 0.0f;
+	}
+
+	void FAudioModulationSystem::UpdateAudioInsightsTraceMessages(const double InElapsed)
+	{
+		const bool bChannelEnabled = UE_TRACE_CHANNELEXPR_IS_ENABLED(AudioMixerChannel);
+		if (!bIsAudioModulationInsightsPluginEnabled || !bChannelEnabled)
 		{
-			for (const auto& [BusId, ControlBusProxy] : RefProxies.Buses)
+			return;
+		}
+		
+		RunCommandOnProcessingThread([this, InElapsed]()
+		{
+			ModulatorActivePingTimeSec -= InElapsed;
+
+			CurrentlyActiveModulatorIds.Reset();
+
+			for (const auto& [GeneratorId, GeneratorProxy] : RefProxies.Generators)
 			{
-				UE_TRACE_LOG(Audio, BusMixRegisterBus, AudioChannel)
-					<< BusMixRegisterBus.DeviceId(AudioDeviceId)
-					<< BusMixRegisterBus.SourceId(ControlBusProxy.GetId())
-					<< BusMixRegisterBus.Timestamp(FPlatformTime::Cycles64())
-					<< BusMixRegisterBus.ModulatingSourceId(BusId)
-					<< BusMixRegisterBus.BusName(*(ControlBusProxy.GetPath().ToString()));
-
-				ControlBusProxy.OnTraceStarted(*this);
+				TraceUpdateModulatorMessage<FModulatorGeneratorProxy>(GeneratorId, GeneratorProxy, EModulatorTraceType::Generator);
+				CurrentlyActiveModulatorIds.Emplace(GeneratorId);
 			}
 
+			for (const auto& [PatchId, ParamaterPatchProxy] : RefProxies.Patches)
+			{
+				const int32 NumStages = ParamaterPatchProxy.InputProxies.Num();
+
+				ContributingModulatorIds.Reset();
+				ContributingModulatorValues.Reset();
+
+				for (const FModulationInputProxy& Input : ParamaterPatchProxy.InputProxies)
+				{
+					ContributingModulatorIds.Emplace(Input.BusHandle.GetId());
+					ContributingModulatorValues.Emplace(Input.ModStageValue);
+				}
+
+				TraceUpdateModulatorMessage<FModulationPatchRefProxy>(PatchId, ParamaterPatchProxy, EModulatorTraceType::ParameterPatch, ContributingModulatorIds, ContributingModulatorValues);
+				CurrentlyActiveModulatorIds.Emplace(PatchId);
+			}
+
+			// For keeping track of the relationship between ControlBusMix (Id, value) and a ControlBus. Allowing ControlBuses to send through ControlBusMixes as contributors.
+			ControlBusMixContributingDataContainer.Reset();
+
 			for (const auto& [BusMixId, ModulatorBusMixProxy] : RefProxies.BusMixes)
 			{
-				UE_TRACE_LOG(Audio, BusMixActivate, AudioChannel)
-					<< BusMixActivate.DeviceId(AudioDeviceId)
-					<< BusMixActivate.SourceId(BusMixId)
-					<< BusMixActivate.Timestamp(FPlatformTime::Cycles64())
-					<< BusMixActivate.Name(*(ModulatorBusMixProxy.GetPath().ToString()));
+				const FModulatorBusMixProxy::EStatus CurrentStatus = ModulatorBusMixProxy.GetStatus();
+
+				switch (CurrentStatus)
+				{
+					case FModulatorBusMixProxy::EStatus::Enabled:
+					case FModulatorBusMixProxy::EStatus::Stopping:
+					{
+						const int32 NumStages = ModulatorBusMixProxy.Stages.Num();
+
+						ContributingModulatorIds.Reset();
+						ContributingModulatorValues.Reset();
+
+						for (const auto& [BusId, BusMixStageProxy] : ModulatorBusMixProxy.Stages)
+						{
+							ContributingModulatorIds.Emplace(BusId);
+							ContributingModulatorValues.Emplace(BusMixStageProxy.Value.GetCurrentValue());
+
+							// Note: Intentionally do not empty this each BusMix iteration, as we require all data to be available for ControlBus contributing logic below
+							ControlBusMixContributingDataContainer.Emplace(AudioInsightsControlBusMixData{BusMixId, BusId, BusMixStageProxy.Value.GetCurrentValue()});
+						}
+
+						TraceUpdateModulatorMessage<FModulatorBusMixProxy>(BusMixId, ModulatorBusMixProxy, EModulatorTraceType::ControlBusMix, ContributingModulatorIds, ContributingModulatorValues);
+						CurrentlyActiveModulatorIds.Emplace(BusMixId);
+					}
+					break;
+					
+					default: break;
+				}
+			}
+
+			for (const auto& [BusId, ControlBusProxy] : RefProxies.Buses)
+			{
+				ContributingModulatorIds.Reset();
+				ContributingModulatorValues.Reset();
+
+				for (const FGeneratorHandle& GeneratorHandle : ControlBusProxy.GetGeneratorHandles())
+				{
+					ContributingModulatorIds.Emplace(GeneratorHandle.GetId());
+					ContributingModulatorValues.Emplace(GeneratorHandle.FindProxy().GetValue());
+				}
+
+				// Note: Multiple ControlBusMixes can contain the same ControlBus, therefore we must iterate over all ControlBusMixContributingDataContainer
+				for (const AudioInsightsControlBusMixData& AudioInsightsControlBusMixData : ControlBusMixContributingDataContainer)
+				{
+					if (BusId == AudioInsightsControlBusMixData.ControlBusId)
+					{
+						ContributingModulatorIds.Emplace(AudioInsightsControlBusMixData.ControlBusMixId);
+						ContributingModulatorValues.Emplace(AudioInsightsControlBusMixData.Value);
+					}
+				}
+
+				TraceUpdateModulatorMessage<FControlBusProxy>(BusId, ControlBusProxy, EModulatorTraceType::ControlBus, ContributingModulatorIds, ContributingModulatorValues);
+				CurrentlyActiveModulatorIds.Emplace(BusId);
+			}
+
+			for (const Audio::FModulatorId PreviouslyActiveModulatorId : PreviouslyActiveModulatorIds)
+			{
+				if (!CurrentlyActiveModulatorIds.Contains(PreviouslyActiveModulatorId))
+				{
+					TraceDeactivateModulatorMessage(PreviouslyActiveModulatorId);
+				}
+			}
+
+			PreviouslyActiveModulatorIds = CurrentlyActiveModulatorIds;
+
+			if (ModulatorActivePingTimeSec <= 0.0f)
+			{
+				ModulatorActivePingTimeSec = ModulatorActivePingIntervalTimeSec;
 			}
 		});
 	}
+
+	void FAudioModulationSystem::TraceActivateModulatorMessage(const Audio::FModulatorId InModulatorId,
+		const EModulatorTraceType InModulatorTraceType,
+		const FString& InModulatorName)
+	{
+		ensure(IsInProcessingThread());
+
+		UE_TRACE_LOG(Audio, ActivateModulatorTraceMessage, AudioMixerChannel)
+			<< ActivateModulatorTraceMessage.DeviceId(GetAudioDeviceId())
+			<< ActivateModulatorTraceMessage.Timestamp(FPlatformTime::Cycles64())
+			<< ActivateModulatorTraceMessage.ModulatorId(InModulatorId)
+			<< ActivateModulatorTraceMessage.ModulatorType(static_cast<Audio::FModulatorTypeId>(InModulatorTraceType))
+			<< ActivateModulatorTraceMessage.Name(*InModulatorName);
+	}
+
+	template <typename TModulatorProxy>
+	void FAudioModulationSystem::TraceUpdateModulatorMessage(const Audio::FModulatorId InModulatorId,
+		const TModulatorProxy& InModulatorProxy,
+		const EModulatorTraceType InModulatorTraceType,
+		TArrayView<const Audio::FModulatorId> InContributingModulatorIds /* = TArrayView<const Audio::FModulatorId>() */,
+		TArrayView<const float> InContributingModulatorValues /* = TArrayView<const float>()*/)
+	{
+		ensure(IsInProcessingThread());
+		ensure(InModulatorTraceType != EModulatorTraceType::COUNT);
+
+		if (ModulatorActivePingTimeSec <= 0.0f)
+		{
+			TraceActivateModulatorMessage(InModulatorId, InModulatorTraceType, *(InModulatorProxy.GetPath().ToString()));
+		}
+
+		// Note: ControlBusMixes do not output a single value, nor can they be bypassed. Therefore using a compile time conditional to allow for flexible message data.
+		if constexpr (std::is_same_v<TModulatorProxy, FModulatorBusMixProxy>) 
+		{
+			UE_TRACE_LOG(Audio, UpdateModulatorTraceMessage, AudioMixerChannel)
+			<< UpdateModulatorTraceMessage.DeviceId(GetAudioDeviceId())
+			<< UpdateModulatorTraceMessage.Timestamp(FPlatformTime::Cycles64())
+			<< UpdateModulatorTraceMessage.ModulatorId(InModulatorId)
+			<< UpdateModulatorTraceMessage.ModulatorValue(0) // Note: BusMixes do not output a single value
+			<< UpdateModulatorTraceMessage.IsBypassed(false) // Note: BusMixes can not be bypassed
+			<< UpdateModulatorTraceMessage.ContributingModulatorIds(InContributingModulatorIds.GetData(), InContributingModulatorIds.Num())
+			<< UpdateModulatorTraceMessage.ContributingModulatorValues(InContributingModulatorValues.GetData(), InContributingModulatorValues.Num());
+		}
+		else 
+		{
+			UE_TRACE_LOG(Audio, UpdateModulatorTraceMessage, AudioMixerChannel)
+			<< UpdateModulatorTraceMessage.DeviceId(GetAudioDeviceId())
+			<< UpdateModulatorTraceMessage.Timestamp(FPlatformTime::Cycles64())
+			<< UpdateModulatorTraceMessage.ModulatorId(InModulatorId)
+			<< UpdateModulatorTraceMessage.ModulatorValue(InModulatorProxy.GetValue())
+			<< UpdateModulatorTraceMessage.IsBypassed(InModulatorProxy.IsBypassed())
+			<< UpdateModulatorTraceMessage.ContributingModulatorIds(InContributingModulatorIds.GetData(), InContributingModulatorIds.Num())
+			<< UpdateModulatorTraceMessage.ContributingModulatorValues(InContributingModulatorValues.GetData(), InContributingModulatorValues.Num());
+		}
+	}
+
+	void FAudioModulationSystem::TraceDeactivateModulatorMessage(const Audio::FModulatorId InModulatorId)
+	{
+		ensure(IsInProcessingThread());
+
+		UE_TRACE_LOG(Audio, DeactivateModulatorTraceMessage, AudioMixerChannel)
+			<< DeactivateModulatorTraceMessage.DeviceId(GetAudioDeviceId())
+			<< DeactivateModulatorTraceMessage.Timestamp(FPlatformTime::Cycles64())
+			<< DeactivateModulatorTraceMessage.ModulatorId(InModulatorId);
+	}
 #endif // UE_AUDIO_PROFILERTRACE_ENABLED
 
 	Audio::FModulatorTypeId FAudioModulationSystem::RegisterModulator(Audio::FModulatorHandleId InHandleId, const FControlBusSettings& InSettings)
@@ -1013,14 +1032,6 @@ namespace AudioModulation
 				if (RefModulators.BusMap.Find(BusHandle) == nullptr)
 				{
 					FControlBusProxy& ControlBusProxy = BusHandle.FindProxy();
-
-					for (const FGeneratorHandle& GeneratorHandle : ControlBusProxy.GetGeneratorHandles())
-					{
-						UE_TRACE_LOG(Audio, ModulatingSourceDeactivate, AudioChannel)
-							<< ModulatingSourceDeactivate.DeviceId(AudioDeviceId)
-							<< ModulatingSourceDeactivate.SourceId(GeneratorHandle.GetId())
-							<< ModulatingSourceDeactivate.Timestamp(FPlatformTime::Cycles64());
-					}
 				}
 #endif
 
@@ -1030,12 +1041,6 @@ namespace AudioModulation
 			FGeneratorHandle GeneratorHandle = FGeneratorHandle::Get(static_cast<FGeneratorId>(ModId), RefProxies.Generators);
 			if (UnregisterModulator<FGeneratorHandle>(GeneratorHandle, RefModulators.GeneratorMap, HandleId))
 			{
-#if UE_AUDIO_PROFILERTRACE_ENABLED
-				UE_TRACE_LOG(Audio, ModulatingSourceDeactivate, AudioChannel)
-					<< ModulatingSourceDeactivate.DeviceId(AudioDeviceId)
-					<< ModulatingSourceDeactivate.SourceId(ModId)
-					<< ModulatingSourceDeactivate.Timestamp(FPlatformTime::Cycles64());
-#endif
 				return;
 			}
 		});
diff --git a/Engine/Plugins/Runtime/AudioModulation/Source/AudioModulation/Private/AudioModulationSystem.h b/Engine/Plugins/Runtime/AudioModulation/Source/AudioModulation/Private/AudioModulationSystem.h
--- a/Engine/Plugins/Runtime/AudioModulation/Source/AudioModulation/Private/AudioModulationSystem.h
+++ b/Engine/Plugins/Runtime/AudioModulation/Source/AudioModulation/Private/AudioModulationSystem.h
@@ -60,14 +60,6 @@ namespace AudioModulation
 	};
 } // namespace AudioModulation
 
-#if UE_AUDIO_PROFILERTRACE_ENABLED
-UE_TRACE_EVENT_BEGIN_EXTERN(Audio, ModulatingSourceDeactivate)
-	UE_TRACE_EVENT_FIELD(uint32, DeviceId)
-	UE_TRACE_EVENT_FIELD(uint32, SourceId)
-	UE_TRACE_EVENT_FIELD(double, Timestamp)
-UE_TRACE_EVENT_END()
-#endif
-
 #if WITH_AUDIOMODULATION
 namespace AudioModulation
 {
@@ -272,10 +264,68 @@ namespace AudioModulation
 		Audio::FDeviceId AudioDeviceId = INDEX_NONE;
 
 #if UE_AUDIO_PROFILERTRACE_ENABLED
+
+		/*
+		* Note: The following mirrors data from Engine/Plugins/Runtime/AudioModulationInsights
+		*
+		* Duplication is required as AudioInsights cannot depend on any engine code
+		* As it must be able to run in standalone (where engine code is not available)
+		* 
+		* Caution: Since this means we no longer have a single source of truth
+		* We must ensure this data is kept in sync
+		*
+		* Intentionally not using the EModulatorType within AudioModulationSystem.cpp:
+		* 1. For the reasons above
+		* 2. To allow inclusion of Control Bus Mixes
+		* 3. To avoid being unintentionally stomped over by future engine changes to EModulatorType
+		*/
+
+		// Mirrors EModulatorTraceType: Engine/Plugins/Runtime/AudioModulationInsights/Source/AudioModulationInsights/Private/AudioModulationDataTypes.h
+		enum class EModulatorTraceType : Audio::FModulatorTypeId
+		{
+			ControlBus,
+			ControlBusMix,
+			Generator,
+			ParameterPatch,
+
+			COUNT
+		};
+
 		void OnTraceStarted(FTraceAuxiliary::EConnectionType TraceType, const FString& TraceDestination);
 
-		float ControlBusActivePingTimeSec = 0.0f;
-		static constexpr float ControlBusActivePingIntervalTimeSec = 1.0f;
+		void UpdateAudioInsightsTraceMessages(const double InElapsed);
+
+		void TraceActivateModulatorMessage(const Audio::FModulatorId InModulatorId,
+			const EModulatorTraceType InModulatorTraceType,
+			const FString& InModulatorName);
+
+		template <typename TModulatorProxy>
+		void TraceUpdateModulatorMessage(const Audio::FModulatorId InModulatorId,
+			const TModulatorProxy& InModulatorProxy,
+			const EModulatorTraceType InModulatorTraceType,
+			TArrayView<const Audio::FModulatorId> InContributingModulatorIds = TArrayView<const Audio::FModulatorId>(),
+			TArrayView<const float> InContributingModulatorValues = TArrayView<const float>());
+
+		void TraceDeactivateModulatorMessage(const Audio::FModulatorId InModulatorId);
+
+		bool bIsAudioModulationInsightsPluginEnabled { false };
+
+		float ModulatorActivePingTimeSec { 0.0f };
+		static constexpr float ModulatorActivePingIntervalTimeSec { 1.0f };
+
+		struct AudioInsightsControlBusMixData
+		{
+			Audio::FModulatorId ControlBusMixId;
+			Audio::FModulatorId ControlBusId;
+			float Value;
+		};
+
+		TArray<Audio::FModulatorId, TInlineAllocator<32>> CurrentlyActiveModulatorIds;
+		TArray<Audio::FModulatorId, TInlineAllocator<32>> PreviouslyActiveModulatorIds;
+		
+		TArray<float, TInlineAllocator<32>> ContributingModulatorValues;
+		TArray<Audio::FModulatorId, TInlineAllocator<32>> ContributingModulatorIds;
+		TArray<AudioInsightsControlBusMixData, TInlineAllocator<32>> ControlBusMixContributingDataContainer;
 #endif // UE_AUDIO_PROFILERTRACE_ENABLED
 
 #if !UE_BUILD_SHIPPING
diff --git a/Engine/Plugins/Runtime/AudioModulation/Source/AudioModulation/Private/SoundControlBusMixProxy.cpp b/Engine/Plugins/Runtime/AudioModulation/Source/AudioModulation/Private/SoundControlBusMixProxy.cpp
--- a/Engine/Plugins/Runtime/AudioModulation/Source/AudioModulation/Private/SoundControlBusMixProxy.cpp
+++ b/Engine/Plugins/Runtime/AudioModulation/Source/AudioModulation/Private/SoundControlBusMixProxy.cpp
@@ -4,7 +4,6 @@
 #include "Audio/AudioAddressPattern.h"
 #include "AudioDevice.h"
 #include "AudioDeviceManager.h"
-#include "AudioMixerTrace.h"
 #include "AudioModulation.h"
 #include "AudioModulationLogging.h"
 #include "AudioModulationProfileSerializer.h"
@@ -270,13 +269,6 @@ namespace AudioModulation
 		if (bRequestStop)
 		{
 			Status = EStatus::Stopped;
-			
-#if	UE_AUDIO_PROFILERTRACE_ENABLED
-			UE_TRACE_LOG(Audio, ModulatingSourceDeactivate, AudioChannel)
-				<< ModulatingSourceDeactivate.DeviceId(ModSystem->GetAudioDeviceId())
-				<< ModulatingSourceDeactivate.SourceId(GetId())
-				<< ModulatingSourceDeactivate.Timestamp(FPlatformTime::Cycles64());
-#endif
 		}
 	}
 
diff --git a/Engine/Plugins/Runtime/AudioModulation/Source/AudioModulation/Private/SoundControlBusProxy.cpp b/Engine/Plugins/Runtime/AudioModulation/Source/AudioModulation/Private/SoundControlBusProxy.cpp
--- a/Engine/Plugins/Runtime/AudioModulation/Source/AudioModulation/Private/SoundControlBusProxy.cpp
+++ b/Engine/Plugins/Runtime/AudioModulation/Source/AudioModulation/Private/SoundControlBusProxy.cpp
@@ -1,9 +1,6 @@
 // Copyright Epic Games, Inc. All Rights Reserved.
 #include "SoundControlBusProxy.h"
 
-#include "AudioDevice.h"
-#include "AudioDeviceManager.h"
-#include "AudioMixerTrace.h"
 #include "AudioModulation.h"
 #include "AudioModulationLogging.h"
 #include "AudioModulationSubsystem.h"
@@ -19,38 +16,6 @@ FAutoConsoleVariableRef CVarForceClampNormalizedValues(
 	TEXT("0: Normalized values can go outside [0,1] if the Modulation Parameter allows it. 1: Normalized values are always in [0,1] range."),
 	ECVF_Default);
 
-#if UE_AUDIO_PROFILERTRACE_ENABLED
-UE_TRACE_EVENT_BEGIN(Audio, ControlBusActivate)
-	UE_TRACE_EVENT_FIELD(uint32, DeviceId)
-	UE_TRACE_EVENT_FIELD(uint32, ControlBusId)
-	UE_TRACE_EVENT_FIELD(uint64, Timestamp)
-	UE_TRACE_EVENT_FIELD(UE::Trace::WideString, Name)
-	UE_TRACE_EVENT_FIELD(UE::Trace::WideString, ParamName)
-UE_TRACE_EVENT_END()
-
-UE_TRACE_EVENT_BEGIN(Audio, ControlBusDeactivate)
-	UE_TRACE_EVENT_FIELD(uint32, DeviceId)
-	UE_TRACE_EVENT_FIELD(uint32, ControlBusId)
-	UE_TRACE_EVENT_FIELD(uint64, Timestamp)
-UE_TRACE_EVENT_END()
-
-UE_TRACE_EVENT_BEGIN(Audio, GeneratorRegisterBus)
-	UE_TRACE_EVENT_FIELD(uint32, DeviceId)
-	UE_TRACE_EVENT_FIELD(uint32, SourceId)
-	UE_TRACE_EVENT_FIELD(uint64, Timestamp)
-	UE_TRACE_EVENT_FIELD(uint32, ModulatingSourceId)
-	UE_TRACE_EVENT_FIELD(UE::Trace::WideString, BusName)
-UE_TRACE_EVENT_END()
-
-UE_TRACE_EVENT_BEGIN(Audio, GeneratorActivate)
-	UE_TRACE_EVENT_FIELD(uint32, DeviceId)
-	UE_TRACE_EVENT_FIELD(uint32, SourceId)
-	UE_TRACE_EVENT_FIELD(uint64, Timestamp)
-	UE_TRACE_EVENT_FIELD(UE::Trace::WideString, Name)
-UE_TRACE_EVENT_END()
-
-#endif // UE_AUDIO_PROFILERTRACE_ENABLED
-
 namespace AudioModulation
 {
 	const FBusId InvalidBusId = INDEX_NONE;
@@ -59,24 +24,6 @@ namespace AudioModulation
 	{
 		FAudioModulationSystem& ModSystem = static_cast<FAudioModulationManager&>(InModulation).GetSystem();
 
-#if UE_AUDIO_PROFILERTRACE_ENABLED
-		for (const FModulationGeneratorSettings& GeneratorSetting : GeneratorSettings)
-		{
-			UE_TRACE_LOG(Audio, GeneratorRegisterBus, AudioChannel)
-				<< GeneratorRegisterBus.DeviceId(ModSystem.GetAudioDeviceId())
-				<< GeneratorRegisterBus.SourceId(GetId())
-				<< GeneratorRegisterBus.Timestamp(FPlatformTime::Cycles64())
-				<< GeneratorRegisterBus.ModulatingSourceId(GeneratorSetting.GetId())
-				<< GeneratorRegisterBus.BusName(*(GetPath().ToString()));
-
-			UE_TRACE_LOG(Audio, GeneratorActivate, AudioChannel)
-				<< GeneratorActivate.DeviceId(ModSystem.GetAudioDeviceId())
-				<< GeneratorActivate.SourceId(GeneratorSetting.GetId())
-				<< GeneratorActivate.Timestamp(FPlatformTime::Cycles64())
-				<< GeneratorActivate.Name(*GeneratorSetting.GetPath().ToString());
-		}
-#endif
-
 		return ModSystem.RegisterModulator(HandleId, *this);
 	}
 
@@ -86,19 +33,6 @@ namespace AudioModulation
 		Init(MoveTemp(InSettings));
 	}
 
-	FControlBusProxy::~FControlBusProxy()
-	{
-#if UE_AUDIO_PROFILERTRACE_ENABLED
-		if (ModSystem)
-		{
-			UE_TRACE_LOG(Audio, ControlBusDeactivate, AudioChannel)
-				<< ControlBusDeactivate.DeviceId(static_cast<uint32>(ModSystem->AudioDeviceId))
-				<< ControlBusDeactivate.ControlBusId(static_cast<uint32>(GetId()))
-				<< ControlBusDeactivate.Timestamp(FPlatformTime::Cycles64());
-		}
-#endif // UE_AUDIO_PROFILERTRACE_ENABLED
-	}
-
 	FControlBusProxy& FControlBusProxy::operator =(FControlBusSettings&& InSettings)
 	{
 		Init(MoveTemp(InSettings));
@@ -172,15 +106,6 @@ namespace AudioModulation
 		// Move vs. reset and adding to original array to avoid potentially clearing handles (and thus current Generator state)
 		// and destroying generators if function is called while reinitializing/updating the modulator
 		GeneratorHandles = MoveTemp(NewHandles);
-
-#if UE_AUDIO_PROFILERTRACE_ENABLED
-		UE_TRACE_LOG(Audio, ControlBusActivate, AudioChannel)
-			<< ControlBusActivate.DeviceId(ModSystem->GetAudioDeviceId())
-			<< ControlBusActivate.ControlBusId(InSettings.GetId())
-			<< ControlBusActivate.Timestamp(FPlatformTime::Cycles64())
-			<< ControlBusActivate.Name(*(InSettings.GetPath().ToString()))
-			<< ControlBusActivate.ParamName(*(GetParameterName().ToString()));
-#endif // UE_AUDIO_PROFILERTRACE_ENABLED
 	}
 
 	bool FControlBusProxy::IsBypassed() const
@@ -214,34 +139,6 @@ namespace AudioModulation
 		return bClampToNormalizedRange;		
 	}
 
-#if UE_AUDIO_PROFILERTRACE_ENABLED
-	void AudioModulation::FControlBusProxy::OnTraceStarted(FAudioModulationSystem& InModSystem) const
-	{
-		UE_TRACE_LOG(Audio, ControlBusActivate, AudioChannel)
-			<< ControlBusActivate.DeviceId(InModSystem.GetAudioDeviceId())
-			<< ControlBusActivate.ControlBusId(GetId())
-			<< ControlBusActivate.Timestamp(FPlatformTime::Cycles64())
-			<< ControlBusActivate.Name(*(GetName().ToString()))
-			<< ControlBusActivate.ParamName(*(GetParameterName().ToString()));
-
-		for (const FGeneratorHandle& GeneratorHandle: GeneratorHandles)
-		{
-			UE_TRACE_LOG(Audio, GeneratorRegisterBus, AudioChannel)
-				<< GeneratorRegisterBus.DeviceId(InModSystem.GetAudioDeviceId())
-				<< GeneratorRegisterBus.SourceId(GetId())
-				<< GeneratorRegisterBus.Timestamp(FPlatformTime::Cycles64())
-				<< GeneratorRegisterBus.ModulatingSourceId(GeneratorHandle.GetId())
-				<< GeneratorRegisterBus.BusName(*(GetPath().ToString()));
-
-			UE_TRACE_LOG(Audio, GeneratorActivate, AudioChannel)
-				<< GeneratorActivate.DeviceId(InModSystem.GetAudioDeviceId())
-				<< GeneratorActivate.SourceId(GeneratorHandle.GetId())
-				<< GeneratorActivate.Timestamp(FPlatformTime::Cycles64())
-				<< GeneratorActivate.Name(*(GeneratorHandle.FindProxy().GetPath().ToString()));
-		}
-	}
-#endif // UE_AUDIO_PROFILERTRACE_ENABLED
-
 	void FControlBusProxy::MixIn(const float InValue)
 	{
 		MixValue = Mix(InValue);
diff --git a/Engine/Plugins/Runtime/AudioModulation/Source/AudioModulation/Private/SoundControlBusProxy.h b/Engine/Plugins/Runtime/AudioModulation/Source/AudioModulation/Private/SoundControlBusProxy.h
--- a/Engine/Plugins/Runtime/AudioModulation/Source/AudioModulation/Private/SoundControlBusProxy.h
+++ b/Engine/Plugins/Runtime/AudioModulation/Source/AudioModulation/Private/SoundControlBusProxy.h
@@ -3,7 +3,6 @@
 
 #include "AudioDeviceManager.h"
 #include "IAudioModulation.h"
-#include "ProfilingDebugging/TraceAuxiliary.h"
 #include "SoundControlBus.h"
 #include "SoundModulationParameter.h"
 #include "SoundModulationProxy.h"
@@ -67,8 +66,6 @@ namespace AudioModulation
 	public:
 		FControlBusProxy();
 		FControlBusProxy(FControlBusSettings&& InSettings, FAudioModulationSystem& InModSystem);
-		
-		~FControlBusProxy();
 
 		FControlBusProxy& operator =(FControlBusSettings&& InSettings);
 
@@ -83,10 +80,6 @@ namespace AudioModulation
 		void MixGenerators();
 		void Reset();
 
-#if UE_AUDIO_PROFILERTRACE_ENABLED
-		void OnTraceStarted(FAudioModulationSystem& InModSystem) const;
-#endif // UE_AUDIO_PROFILERTRACE_ENABLED
-
 	private:
 		void Init(FControlBusSettings&& InSettings);
 		float Mix(float ValueA) const;
