Fix EditCondition log error in Modulation Parameter Clamp setting
#Synced-CL 50724195

diff --git a/Engine/Plugins/Runtime/AudioModulation/Source/AudioModulation/Public/SoundModulationParameter.h b/Engine/Plugins/Runtime/AudioModulation/Source/AudioModulation/Public/SoundModulationParameter.h
--- a/Engine/Plugins/Runtime/AudioModulation/Source/AudioModulation/Public/SoundModulationParameter.h
+++ b/Engine/Plugins/Runtime/AudioModulation/Source/AudioModulation/Public/SoundModulationParameter.h
@@ -40,6 +40,7 @@ struct FSoundModulationParameterSettings
 	UPROPERTY(Transient, EditAnywhere, Category = General)
 	float ValueUnit = 1.0f;
 
+	UPROPERTY()
 	bool bShowClampOption = false;
 #endif // WITH_EDITORONLY_DATA
 };
