Sequencer: Control bus input clamping

#Synced-CL 51351972

diff --git a/Engine/Plugins/Runtime/AudioModulation/Source/AudioModulation/Private/MovieSceneAudioControlBusSection.cpp b/Engine/Plugins/Runtime/AudioModulation/Source/AudioModulation/Private/MovieSceneAudioControlBusSection.cpp
--- a/Engine/Plugins/Runtime/AudioModulation/Source/AudioModulation/Private/MovieSceneAudioControlBusSection.cpp
+++ b/Engine/Plugins/Runtime/AudioModulation/Source/AudioModulation/Private/MovieSceneAudioControlBusSection.cpp
@@ -24,11 +24,19 @@ EMovieSceneChannelProxyType UMovieSceneAudioControlBusSection::CacheChannelProxy
 {
 	// Set up the channel proxy
 	FMovieSceneChannelProxyData Channels;
+	check(ControlBus != nullptr);
 
 #if WITH_EDITOR
 	FMovieSceneChannelMetaData Data;
 	Data.SetIdentifiers(TEXT("Control Value"), FText::FromName(ControlBus == nullptr ? TEXT("Control Bus") : ControlBus->GetFName()));
 
+	check(ControlBus->Parameter != nullptr);
+	const float MinValue = ControlBus->Parameter->GetUnitMin();
+	const float MaxValue = ControlBus->Parameter->GetUnitMax();
+
+	Data.PropertyMetaData.Add(TEXT("ClampMin"), FString::SanitizeFloat(MinValue));
+	Data.PropertyMetaData.Add(TEXT("ClampMax"), FString::SanitizeFloat(MaxValue));
+
 	Channels.Add(BusValue, Data, TMovieSceneExternalValue<float>::Make());
 #else
 	Channels.Add(BusValue);
@@ -65,4 +73,7 @@ void UMovieSceneAudioControlBusSection::ImportEntityImpl(UMovieSceneEntitySystem
 void UMovieSceneAudioControlBusSection::SetBus(USoundControlBus* NewBus)
 {
 	ControlBus = NewBus;
+
+	// Set default to be in the min range
+	BusValue.SetDefault(NewBus->Parameter->GetUnitMin());
 }
