Sequencer: Fix control bus tracks to send normalized values when updating

#Synced-CL 51691768

diff --git a/Engine/Plugins/Runtime/AudioModulation/Source/AudioModulation/Private/MovieSceneAudioControlBusSystem.cpp b/Engine/Plugins/Runtime/AudioModulation/Source/AudioModulation/Private/MovieSceneAudioControlBusSystem.cpp
--- a/Engine/Plugins/Runtime/AudioModulation/Source/AudioModulation/Private/MovieSceneAudioControlBusSystem.cpp
+++ b/Engine/Plugins/Runtime/AudioModulation/Source/AudioModulation/Private/MovieSceneAudioControlBusSystem.cpp
@@ -106,6 +106,9 @@ namespace UE::MovieScene
 
 			USoundControlBus* Bus = CastChecked<USoundControlBus>(AudioData.Section->ControlBus);
 
+			// Control busses now use normalized values so we convert from user facing values
+			const double NormalisedValue = Bus->Parameter->ConvertUnitToNormalized(BusValue);
+
 			// If no cached value exists, we need to initialize
 			if (!EvaluationData || !EvaluationData->ControlBusMix->IsValidLowLevel() || EvaluationData->ControlBusMix == nullptr)
 			{
@@ -116,7 +119,7 @@ namespace UE::MovieScene
 					AudioSystem->RemoveAudioControlBusComponentEvaluationData(AudioSection);
 				}
 
-				if (USoundControlBusMix* NewMix = UAudioModulationStatics::CreateBusMixFromValue(World, Bus->GetFName(), {Bus}, BusValue))
+				if (USoundControlBusMix* NewMix = UAudioModulationStatics::CreateBusMixFromValue(World, Bus->GetFName(), {Bus}, NormalisedValue))
 				{
 					EvaluationData = AudioSystem->AddAudioControlBusComponentEvaluationData(InstanceHandle, AudioSection, Bus, NewMix);
 				}
@@ -131,7 +134,7 @@ namespace UE::MovieScene
 					UAudioModulationStatics::ActivateBusMix(World, EvaluationData->ControlBusMix.Get());
 				}
 
-				UAudioModulationStatics::UpdateMixByFilter(World, EvaluationData->ControlBusMix.Get(), Bus->Address, nullptr, Bus->Parameter, BusValue, -1);
+				UAudioModulationStatics::UpdateMixByFilter(World, EvaluationData->ControlBusMix.Get(), Bus->Address, nullptr, Bus->Parameter, NormalisedValue, -1);
 			}		
 		}
 	};
