Sequencer: Control bus tracks clear mixes on stop events

#Synced-CL 51850311

diff --git a/Engine/Plugins/Runtime/AudioModulation/Source/AudioModulation/Private/MovieSceneAudioControlBusSystem.cpp b/Engine/Plugins/Runtime/AudioModulation/Source/AudioModulation/Private/MovieSceneAudioControlBusSystem.cpp
--- a/Engine/Plugins/Runtime/AudioModulation/Source/AudioModulation/Private/MovieSceneAudioControlBusSystem.cpp
+++ b/Engine/Plugins/Runtime/AudioModulation/Source/AudioModulation/Private/MovieSceneAudioControlBusSystem.cpp
@@ -135,7 +135,14 @@ namespace UE::MovieScene
 				}
 
 				UAudioModulationStatics::UpdateMixByFilter(World, EvaluationData->ControlBusMix.Get(), Bus->Address, nullptr, Bus->Parameter, NormalisedValue, -1);
-			}		
+			}
+
+			// Mixes are only cleared on stopped and not paused because the sequence is still considered active and applying mixes in a paused state
+			if (Instance.GetContext().GetStatus() == EMovieScenePlayerStatus::Stopped)
+			{
+				// On stop clear mixes to allow any other playback in engine to function without those mixes impacting
+				AudioSystem->ClearBusMixes();
+			}
 		}
 	};
 }
