Sequencer: Set control bus track cvar to default to true

#Synced-CL 52364776

diff --git a/Engine/Plugins/Runtime/AudioModulation/Source/AudioModulationEditor/Private/AudioControlBusBaseTrackEditor.cpp b/Engine/Plugins/Runtime/AudioModulation/Source/AudioModulationEditor/Private/AudioControlBusBaseTrackEditor.cpp
--- a/Engine/Plugins/Runtime/AudioModulation/Source/AudioModulationEditor/Private/AudioControlBusBaseTrackEditor.cpp
+++ b/Engine/Plugins/Runtime/AudioModulation/Source/AudioModulationEditor/Private/AudioControlBusBaseTrackEditor.cpp
@@ -4,12 +4,12 @@
 
 #include "AssetRegistry/AssetData.h"
 
-static bool bUseControlBusTracks = false;
+static bool bUseControlBusTracks = true;
 
 static FAutoConsoleVariableRef CVarControlBusTracks(
 	TEXT("au.ControlBusTracks"),
 	bUseControlBusTracks,
-	TEXT("Enable the control bus track and control bus mix track in sequencer using the audio modulation plugin. Default disables adding new control bus tracks\n")
+	TEXT("Enable the control bus track and control bus mix track in sequencer using the audio modulation plugin. Default enables adding new control bus tracks\n")
 	TEXT("0: Disable control bus tracks\n")
 	TEXT("1: Enable control bus tracks\n"),
 	ECVF_Default
